Audio: CRASH: Crash playing TimeSynth that hasn't been activated

diff --git a/Engine/Plugins/Runtime/TimeSynth/Source/TimeSynth/Private/TimeSynthComponent.cpp b/Engine/Plugins/Runtime/TimeSynth/Source/TimeSynth/Private/TimeSynthComponent.cpp
--- a/Engine/Plugins/Runtime/TimeSynth/Source/TimeSynth/Private/TimeSynthComponent.cpp
+++ b/Engine/Plugins/Runtime/TimeSynth/Source/TimeSynth/Private/TimeSynthComponent.cpp
@@ -53,12 +53,23 @@ void UTimeSynthComponent::OnRegister()
 	Super::OnRegister();
 
 	SetComponentTickEnabled(true);
-	RegisterComponent();
+
+	if (!IsRegistered())
+	{
+		RegisterComponent();
+	}
 }
 
 void UTimeSynthComponent::OnUnregister()
 {
 	Super::OnUnregister();
+
+	SetComponentTickEnabled(false);
+
+	if (IsRegistered())
+	{
+		UnregisterComponent();
+	}
 }
 
 bool UTimeSynthComponent::IsReadyForFinishDestroy()
@@ -564,6 +575,11 @@ FTimeSynthClipHandle UTimeSynthComponent::PlayClip(UTimeSynthClip* InClip, UTime
 		return FTimeSynthClipHandle();
 	}
 
+	if (!bIsActive)
+	{
+		SetActive(true);
+	}
+
 	// Get this time synth components transform
 	const FTransform& ThisComponentTransform = GetComponentTransform();
 
