Merge Time Synth Updates from Dev-Audio:
- Add asset action to create time synth clip from sound wave(s)
- Fix crash when passing a null time synth volume group to SetVolumeGroup
- Parameterize max pool size

diff --git a/Engine/Plugins/Runtime/TimeSynth/Source/TimeSynth/Classes/TimeSynthComponent.h b/Engine/Plugins/Runtime/TimeSynth/Source/TimeSynth/Classes/TimeSynthComponent.h
--- a/Engine/Plugins/Runtime/TimeSynth/Source/TimeSynth/Classes/TimeSynthComponent.h
+++ b/Engine/Plugins/Runtime/TimeSynth/Source/TimeSynth/Classes/TimeSynthComponent.h
@@ -242,7 +242,7 @@ public:
 	float DefaultVolume;
 };
 
-UCLASS(ClassGroup = Synth, meta = (BlueprintSpawnableComponent))
+UCLASS(BlueprintType, ClassGroup = Synth, meta = (BlueprintSpawnableComponent))
 class TIMESYNTH_API UTimeSynthClip : public UObject
 {
 	GENERATED_BODY()
@@ -530,6 +530,10 @@ public:
 	UPROPERTY(EditAnywhere, BlueprintReadWrite, Category = "Filter", meta = (EditCondition = "bisEnvelopeFollowerEnabled"))
 	FTimeSynthEnvelopeFollowerSettings EnvelopeFollowerSettings;
 
+	// Set the Max Pool Size
+	UPROPERTY(EditAnywhere, BlueprintReadWrite, Category = "TimeSynth", meta = (UIMin = 20, UIMax = 50, ClampMin = 20))
+	int32 MaxPoolSize;
+
 	// Register an event to respond to a specific quantization event
 	UFUNCTION(BlueprintCallable, Category = "Quantization", meta = (WorldContext = "WorldContextObject"))
 	void AddQuantizationEventDelegate(ETimeSynthEventQuantization QuantizationType, const FOnQuantizationEventBP& OnQuantizationEvent);
diff --git a/Engine/Plugins/Runtime/TimeSynth/Source/TimeSynth/Private/TimeSynthComponent.cpp b/Engine/Plugins/Runtime/TimeSynth/Source/TimeSynth/Private/TimeSynthComponent.cpp
--- a/Engine/Plugins/Runtime/TimeSynth/Source/TimeSynth/Private/TimeSynthComponent.cpp
+++ b/Engine/Plugins/Runtime/TimeSynth/Source/TimeSynth/Private/TimeSynthComponent.cpp
@@ -14,6 +14,7 @@ void FTimeSynthEventListener::OnEvent(Audio::EEventQuantization EventQuantizatio
 
 UTimeSynthComponent::UTimeSynthComponent(const FObjectInitializer& ObjectInitializer) 
 	: Super(ObjectInitializer)
+	, MaxPoolSize(20)
 	, TimeSynthEventListener(this)
 {
 	PrimaryComponentTick.bCanEverTick = true;
@@ -323,7 +324,7 @@ bool UTimeSynthComponent::Init(int32& InSampleRate)
 	SetQuantizationSettings(QuantizationSettings);
 
 	// Create a pool of playing clip runtime infos
-	CurrentPoolSize = 20;
+	CurrentPoolSize = MaxPoolSize;
 
 	PlayingClipsPool_AudioRenderThread.AddDefaulted(CurrentPoolSize);
 	FreePlayingClipIndices_AudioRenderThread.AddDefaulted(CurrentPoolSize);
@@ -877,6 +878,12 @@ void UTimeSynthComponent::SetVolumeGroupInternal(FVolumeGroupData& InData, float
 
 void UTimeSynthComponent::SetVolumeGroup(UTimeSynthVolumeGroup* InVolumeGroup, float VolumeDb, float FadeTimeSec)
 {
+	if (!InVolumeGroup)
+	{
+		UE_LOG(LogTimeSynth, Warning, TEXT("Failed to set volume group. Null UTimeSynthVolumeGroup object."));
+		return;
+	}
+
 	VolumeGroupUniqueId Id = InVolumeGroup->GetUniqueID();
 	FVolumeGroupData* VolumeGroup = VolumeGroupData.Find(Id);
 
diff --git a/Engine/Plugins/Runtime/TimeSynth/Source/TimeSynthEditor/Classes/TimeSynthClip.h b/Engine/Plugins/Runtime/TimeSynth/Source/TimeSynthEditor/Classes/TimeSynthClip.h
--- a/Engine/Plugins/Runtime/TimeSynth/Source/TimeSynthEditor/Classes/TimeSynthClip.h
+++ b/Engine/Plugins/Runtime/TimeSynth/Source/TimeSynthEditor/Classes/TimeSynthClip.h
@@ -2,13 +2,15 @@
 
 #pragma once
 
-#include "CoreMinimal.h"
-#include "Stats/Stats.h"
 #include "AssetToolsModule.h"
-#include "Factories/Factory.h"
 #include "AssetTypeActions_Base.h"
+#include "CoreMinimal.h"
+#include "Factories/Factory.h"
+#include "Stats/Stats.h"
 #include "TimeSynthClip.generated.h"
 
+class USoundWave;
+
 class FAssetTypeActions_TimeSynthClip : public FAssetTypeActions_Base
 {
 public:
@@ -26,7 +28,6 @@ class UTimeSynthClipFactory : public UFactory
 	//~ Begin UFactory Interface
 	virtual UObject* FactoryCreateNew(UClass* Class, UObject* InParent, FName Name, EObjectFlags Flags, UObject* Context, FFeedbackContext* Warn) override;
 	//~ Begin UFactory Interface	
-};
-
-
 
+	TArray<TWeakObjectPtr<USoundWave>> SoundWaves;
+};
diff --git a/Engine/Plugins/Runtime/TimeSynth/Source/TimeSynthEditor/Private/TimeSynthClip.cpp b/Engine/Plugins/Runtime/TimeSynth/Source/TimeSynthEditor/Private/TimeSynthClip.cpp
--- a/Engine/Plugins/Runtime/TimeSynth/Source/TimeSynthEditor/Private/TimeSynthClip.cpp
+++ b/Engine/Plugins/Runtime/TimeSynth/Source/TimeSynthEditor/Private/TimeSynthClip.cpp
@@ -2,9 +2,11 @@
 
 #include "TimeSynthClip.h"
 #include "TimeSynthComponent.h"
+#include "Sound/SoundWave.h"
 #include "UObject/ObjectMacros.h"
 #include "UObject/Object.h"
 
+
 UClass* FAssetTypeActions_TimeSynthClip::GetSupportedClass() const
 {
 	return UTimeSynthClip::StaticClass();
@@ -24,5 +26,15 @@ UObject* UTimeSynthClipFactory::FactoryCreateNew(UClass* Class, UObject* InParen
 {
 	UTimeSynthClip* NewTimeSynthClip = NewObject<UTimeSynthClip>(InParent, InName, Flags);
 
+	for (TWeakObjectPtr<USoundWave>& SoundWave : SoundWaves)
+	{
+		if (SoundWave.IsValid())
+		{
+			FTimeSynthClipSound Sound;
+			Sound.SoundWave = SoundWave.Get();
+			NewTimeSynthClip->Sounds.Add(Sound);
+		}
+	}
+	
 	return NewTimeSynthClip;
 }
\ No newline at end of file
diff --git a/Engine/Plugins/Runtime/TimeSynth/Source/TimeSynthEditor/Private/TimeSynthEditorModule.cpp b/Engine/Plugins/Runtime/TimeSynth/Source/TimeSynthEditor/Private/TimeSynthEditorModule.cpp
--- a/Engine/Plugins/Runtime/TimeSynth/Source/TimeSynthEditor/Private/TimeSynthEditorModule.cpp
+++ b/Engine/Plugins/Runtime/TimeSynth/Source/TimeSynthEditor/Private/TimeSynthEditorModule.cpp
@@ -9,15 +9,24 @@
 #include "Factories/Factory.h"
 #include "AssetTypeActions_Base.h"
 #include "TimeSynthClip.h"
+#include "TimeSynthSoundWaveAssetActionExtender.h"
 #include "TimeSynthVolumeGroup.h"
+
 #include "AudioEditorModule.h"
 
+
 IMPLEMENT_MODULE(FTimeSynthEditorModule, TimeSynthEditor)
 
 void FTimeSynthEditorModule::StartupModule()
 {
+	TimeSynthSoundWaveAssetActionExtender = TSharedPtr<ISoundWaveAssetActionExtensions>(new FTimeSynthSoundWaveAssetActionExtender());
+
 	IAssetTools& AssetTools = FModuleManager::LoadModuleChecked<FAssetToolsModule>("AssetTools").Get();
 
+	// Now that we've loaded this module, we need to register our effect preset actions
+	IAudioEditorModule* AudioEditorModule = &FModuleManager::LoadModuleChecked<IAudioEditorModule>("AudioEditor");
+	AudioEditorModule->AddSoundWaveActionExtender(TimeSynthSoundWaveAssetActionExtender);
+
 	AssetTools.RegisterAssetTypeActions(MakeShareable(new FAssetTypeActions_TimeSynthClip));
 	AssetTools.RegisterAssetTypeActions(MakeShareable(new FAssetTypeActions_TimeSynthVolumeGroup));
 }
@@ -25,5 +34,3 @@ void FTimeSynthEditorModule::StartupModule()
 void FTimeSynthEditorModule::ShutdownModule()
 {
 }
-
-
diff --git a/Engine/Plugins/Runtime/TimeSynth/Source/TimeSynthEditor/Private/TimeSynthSoundWaveAssetActionExtender.cpp b/Engine/Plugins/Runtime/TimeSynth/Source/TimeSynthEditor/Private/TimeSynthSoundWaveAssetActionExtender.cpp
new file mode 100644
--- /dev/null
+++ b/Engine/Plugins/Runtime/TimeSynth/Source/TimeSynthEditor/Private/TimeSynthSoundWaveAssetActionExtender.cpp
@@ -0,0 +1,77 @@
+// Copyright 1998-2019 Epic Games, Inc. All Rights Reserved.
+
+#include "TimeSynthSoundWaveAssetActionExtender.h"
+
+#include "Framework/Commands/UIAction.h"
+#include "Framework/Commands/UICommandList.h"
+#include "Framework/MultiBox/MultiBoxBuilder.h"
+#include "AssetTypeActions_Base.h"
+#include "IContentBrowserSingleton.h"
+#include "ContentBrowserModule.h"
+#include "EditorStyleSet.h"
+#include "Sound/SoundWave.h"
+#include "TimeSynthClip.h"
+#include "TimeSynthComponent.h"
+
+#define LOCTEXT_NAMESPACE "AssetTypeActions"
+
+namespace
+{
+	const FString DefaultSuffix = TEXT("_TSC");
+} // namespace <>
+
+void FTimeSynthSoundWaveAssetActionExtender::GetExtendedActions(const TArray<TWeakObjectPtr<USoundWave>>& InSounds, FMenuBuilder& MenuBuilder)
+{
+	const TAttribute<FText> Label = LOCTEXT("SoundWave_CreateTimeSynthClip", "Create Time Synth Clip(s)");
+	const TAttribute<FText> ToolTip = LOCTEXT("SoundWave_CreateTimeSynthClipToolTip", "Creates time synth clip per sound wave selected");
+	const FSlateIcon Icon = FSlateIcon(FEditorStyle::GetStyleSetName(), "ClassIcon.SoundCue");
+	const FUIAction UIAction = FUIAction(FExecuteAction::CreateSP(this, &FTimeSynthSoundWaveAssetActionExtender::ExecuteCreateTimeSyncClip, InSounds), FCanExecuteAction());
+
+	MenuBuilder.AddMenuEntry(Label, ToolTip, Icon, UIAction);
+
+	const TAttribute<FText> LabelSet = LOCTEXT("SoundWave_CreateTimeSynthClipSet", "Create Time Synth Clip Set");
+	const TAttribute<FText> ToolTipSet = LOCTEXT("SoundWave_CreateTimeSynthClipSetToolTip", "Creates time synth clip adding all selected sound waves to single clip as set.");
+	const FUIAction UIActionSet = FUIAction(FExecuteAction::CreateSP(this, &FTimeSynthSoundWaveAssetActionExtender::ExecuteCreateTimeSyncClipSet, InSounds), FCanExecuteAction());
+
+	MenuBuilder.AddMenuEntry(LabelSet, ToolTipSet, Icon, UIActionSet);
+}
+
+void FTimeSynthSoundWaveAssetActionExtender::ExecuteCreateTimeSyncClip(TArray<TWeakObjectPtr<USoundWave>> SoundWaves)
+{
+	FAssetToolsModule& AssetToolsModule = FModuleManager::Get().LoadModuleChecked<FAssetToolsModule>("AssetTools");
+
+	UTimeSynthClipFactory* Factory = NewObject<UTimeSynthClipFactory>();
+	for (TWeakObjectPtr<USoundWave>& SoundWave : SoundWaves)
+	{
+		FString Name;
+		FString PackagePath;
+
+		AssetToolsModule.Get().CreateUniqueAssetName(SoundWave->GetOutermost()->GetName(), DefaultSuffix, PackagePath, Name);
+
+		Factory->SoundWaves.Reset();
+		Factory->SoundWaves.Add(SoundWave);
+
+		AssetToolsModule.Get().CreateAsset(Name, FPackageName::GetLongPackagePath(PackagePath), UTimeSynthClip::StaticClass(), Factory);
+	}
+}
+
+void FTimeSynthSoundWaveAssetActionExtender::ExecuteCreateTimeSyncClipSet(TArray<TWeakObjectPtr<USoundWave>> SoundWaves)
+{
+	FString PackagePath;
+	FString Name;
+
+	if (SoundWaves.Num() == 0 || !SoundWaves[0].IsValid())
+	{
+		return;
+	}
+
+	USoundWave& SoundWave = *SoundWaves[0].Get();
+	FAssetToolsModule& AssetToolsModule = FModuleManager::Get().LoadModuleChecked<FAssetToolsModule>("AssetTools");
+	AssetToolsModule.Get().CreateUniqueAssetName(SoundWave.GetOutermost()->GetName(), DefaultSuffix, PackagePath, Name);
+
+	UTimeSynthClipFactory* Factory = NewObject<UTimeSynthClipFactory>();
+	Factory->SoundWaves = SoundWaves;
+
+	AssetToolsModule.Get().CreateAsset(Name, FPackageName::GetLongPackagePath(PackagePath), UTimeSynthClip::StaticClass(), Factory);
+}
+#undef LOCTEXT_NAMESPACE
\ No newline at end of file
diff --git a/Engine/Plugins/Runtime/TimeSynth/Source/TimeSynthEditor/Public/TimeSynthEditorModule.h b/Engine/Plugins/Runtime/TimeSynth/Source/TimeSynthEditor/Public/TimeSynthEditorModule.h
--- a/Engine/Plugins/Runtime/TimeSynth/Source/TimeSynthEditor/Public/TimeSynthEditorModule.h
+++ b/Engine/Plugins/Runtime/TimeSynth/Source/TimeSynthEditor/Public/TimeSynthEditorModule.h
@@ -1,10 +1,11 @@
 // Copyright 1998-2019 Epic Games, Inc. All Rights Reserved.
-
 #pragma once
 
+#include "AudioEditorModule.h"
 #include "Modules/ModuleInterface.h"
 #include "Modules/ModuleManager.h"
 
+
 class FTimeSynthEditorModule : public IModuleInterface
 {
 public:
@@ -12,4 +13,7 @@ public:
 	// IModuleInterface interface
 	virtual void StartupModule() override;
 	virtual void ShutdownModule() override;
+
+private:
+	TSharedPtr<ISoundWaveAssetActionExtensions> TimeSynthSoundWaveAssetActionExtender;
 };
diff --git a/Engine/Plugins/Runtime/TimeSynth/Source/TimeSynthEditor/Public/TimeSynthSoundWaveAssetActionExtender.h b/Engine/Plugins/Runtime/TimeSynth/Source/TimeSynthEditor/Public/TimeSynthSoundWaveAssetActionExtender.h
new file mode 100644
--- /dev/null
+++ b/Engine/Plugins/Runtime/TimeSynth/Source/TimeSynthEditor/Public/TimeSynthSoundWaveAssetActionExtender.h
@@ -0,0 +1,19 @@
+// Copyright 1998-2019 Epic Games, Inc. All Rights Reserved.
+
+#pragma once
+
+#include "AudioEditorModule.h"
+
+
+class FTimeSynthSoundWaveAssetActionExtender :	public ISoundWaveAssetActionExtensions,
+												public TSharedFromThis<FTimeSynthSoundWaveAssetActionExtender>
+{
+public:
+	virtual ~FTimeSynthSoundWaveAssetActionExtender() {}
+
+	virtual void GetExtendedActions(const TArray<TWeakObjectPtr<USoundWave>>& InObjects, FMenuBuilder& MenuBuilder) override;
+
+	void ExecuteCreateTimeSyncClip(TArray<TWeakObjectPtr<USoundWave>> Objects);
+	void ExecuteCreateTimeSyncClipSet(TArray<TWeakObjectPtr<USoundWave>> SoundWaves);
+};
+
diff --git a/Engine/Plugins/Runtime/TimeSynth/Source/TimeSynthEditor/TimeSynthEditor.Build.cs b/Engine/Plugins/Runtime/TimeSynth/Source/TimeSynthEditor/TimeSynthEditor.Build.cs
--- a/Engine/Plugins/Runtime/TimeSynth/Source/TimeSynthEditor/TimeSynthEditor.Build.cs
+++ b/Engine/Plugins/Runtime/TimeSynth/Source/TimeSynthEditor/TimeSynthEditor.Build.cs
@@ -15,6 +15,9 @@ namespace UnrealBuildTool.Rules
 					"AudioEditor",
                     "TimeSynth",
 					"AudioMixer",
+                    "EditorStyle",
+                    "Slate",
+                    "SlateCore",
 				}
 			);
 
