Fixing timesynth compressor value.

diff --git a/Engine/Plugins/Runtime/TimeSynth/Source/TimeSynth/Private/TimeSynthComponent.cpp b/Engine/Plugins/Runtime/TimeSynth/Source/TimeSynth/Private/TimeSynthComponent.cpp
--- a/Engine/Plugins/Runtime/TimeSynth/Source/TimeSynth/Private/TimeSynthComponent.cpp
+++ b/Engine/Plugins/Runtime/TimeSynth/Source/TimeSynth/Private/TimeSynthComponent.cpp
@@ -306,7 +306,7 @@ bool UTimeSynthComponent::Init(int32& InSampleRate)
 	DynamicsProcessor.SetLookaheadMsec(3.0f);
 	DynamicsProcessor.SetAttackTime(5.0f);
 	DynamicsProcessor.SetReleaseTime(100.0f);
-	DynamicsProcessor.SetThreshold(-15.0f);
+	DynamicsProcessor.SetThreshold(-1.0f);
 	DynamicsProcessor.SetRatio(5.0f);
 	DynamicsProcessor.SetKneeBandwidth(10.0f);
 	DynamicsProcessor.SetInputGain(0.0f);
