Fix build break

diff --git a/Engine/Plugins/Runtime/TimeSynth/Source/TimeSynth/Classes/TimeSynthComponent.h b/Engine/Plugins/Runtime/TimeSynth/Source/TimeSynth/Classes/TimeSynthComponent.h
--- a/Engine/Plugins/Runtime/TimeSynth/Source/TimeSynth/Classes/TimeSynthComponent.h
+++ b/Engine/Plugins/Runtime/TimeSynth/Source/TimeSynth/Classes/TimeSynthComponent.h
@@ -565,7 +565,7 @@ private:
 	void UpdateEnvelopeFollower();
 	void ShutdownPlayingClips();
 
-	Audio::SpectrumAnalyzerSettings::EFFTSize GetFFTSize(ETimeSynthFFTSize InSize) const;
+	Audio::FSpectrumAnalyzerSettings::EFFTSize GetFFTSize(ETimeSynthFFTSize InSize) const;
 
 	// Defines type for a volume group ID
 	typedef uint32 VolumeGroupUniqueId;
diff --git a/Engine/Plugins/Runtime/TimeSynth/Source/TimeSynth/Private/TimeSynthComponent.cpp b/Engine/Plugins/Runtime/TimeSynth/Source/TimeSynth/Private/TimeSynthComponent.cpp
--- a/Engine/Plugins/Runtime/TimeSynth/Source/TimeSynth/Private/TimeSynthComponent.cpp
+++ b/Engine/Plugins/Runtime/TimeSynth/Source/TimeSynth/Private/TimeSynthComponent.cpp
@@ -140,7 +140,7 @@ void UTimeSynthComponent::SetEnvelopeFollowerEnabled(bool bInIsEnabled)
 	});
 }
 
-Audio::SpectrumAnalyzerSettings::EFFTSize UTimeSynthComponent::GetFFTSize(ETimeSynthFFTSize InSize) const
+Audio::FSpectrumAnalyzerSettings::EFFTSize UTimeSynthComponent::GetFFTSize(ETimeSynthFFTSize InSize) const
 {
 	switch (InSize)
 	{
