Fix build break

diff --git a/Engine/Plugins/Runtime/TimeSynth/Source/TimeSynth/Private/TimeSynthComponent.cpp b/Engine/Plugins/Runtime/TimeSynth/Source/TimeSynth/Private/TimeSynthComponent.cpp
--- a/Engine/Plugins/Runtime/TimeSynth/Source/TimeSynth/Private/TimeSynthComponent.cpp
+++ b/Engine/Plugins/Runtime/TimeSynth/Source/TimeSynth/Private/TimeSynthComponent.cpp
@@ -144,19 +144,19 @@ Audio::FSpectrumAnalyzerSettings::EFFTSize UTimeSynthComponent::GetFFTSize(ETime
 {
 	switch (InSize)
 	{
-		case ETimeSynthFFTSize::Min_64: return Audio::SpectrumAnalyzerSettings::EFFTSize::Min_64;
-		case ETimeSynthFFTSize::Small_256: return Audio::SpectrumAnalyzerSettings::EFFTSize::Small_256;
-		case ETimeSynthFFTSize::Medium_512: return Audio::SpectrumAnalyzerSettings::EFFTSize::Medium_512;
-		case ETimeSynthFFTSize::Large_1024: return Audio::SpectrumAnalyzerSettings::EFFTSize::Large_1024;
+		case ETimeSynthFFTSize::Min_64: return Audio::FSpectrumAnalyzerSettings::EFFTSize::Min_64;
+		case ETimeSynthFFTSize::Small_256: return Audio::FSpectrumAnalyzerSettings::EFFTSize::Small_256;
+		case ETimeSynthFFTSize::Medium_512: return Audio::FSpectrumAnalyzerSettings::EFFTSize::Medium_512;
+		case ETimeSynthFFTSize::Large_1024: return Audio::FSpectrumAnalyzerSettings::EFFTSize::Large_1024;
 		break;
 	}
 	// return default
-	return Audio::SpectrumAnalyzerSettings::EFFTSize::Medium_512;
+	return Audio::FSpectrumAnalyzerSettings::EFFTSize::Medium_512;
 }
 
 void UTimeSynthComponent::SetFFTSize(ETimeSynthFFTSize InFFTSize)
 {
-	Audio::SpectrumAnalyzerSettings::EFFTSize NewFFTSize = GetFFTSize(InFFTSize);
+	Audio::FSpectrumAnalyzerSettings::EFFTSize NewFFTSize = GetFFTSize(InFFTSize);
 
 	SynthCommand([this, NewFFTSize]
 	{
