Fix modules that cause custom PCHs to be generated.

diff --git a/Engine/Plugins/Runtime/TimeSynth/Source/TimeSynthEditor/TimeSynthEditor.Build.cs b/Engine/Plugins/Runtime/TimeSynth/Source/TimeSynthEditor/TimeSynthEditor.Build.cs
--- a/Engine/Plugins/Runtime/TimeSynth/Source/TimeSynthEditor/TimeSynthEditor.Build.cs
+++ b/Engine/Plugins/Runtime/TimeSynth/Source/TimeSynthEditor/TimeSynthEditor.Build.cs
@@ -6,8 +6,6 @@ namespace UnrealBuildTool.Rules
 	{
         public TimeSynthEditor(ReadOnlyTargetRules Target) : base(Target)
 		{
-            OptimizeCode = CodeOptimization.Never;
-
             PublicDependencyModuleNames.AddRange(
 				new string[] {
                     "Core",
