Fix ignored bool in TimeSynthComponent

diff --git a/Engine/Plugins/Runtime/TimeSynth/Source/TimeSynth/Classes/TimeSynthComponent.h b/Engine/Plugins/Runtime/TimeSynth/Source/TimeSynth/Classes/TimeSynthComponent.h
--- a/Engine/Plugins/Runtime/TimeSynth/Source/TimeSynth/Classes/TimeSynthComponent.h
+++ b/Engine/Plugins/Runtime/TimeSynth/Source/TimeSynth/Classes/TimeSynthComponent.h
@@ -179,6 +179,12 @@ struct TIMESYNTH_API FTimeSynthTimeDef
 		: NumBars(1)
 		, NumBeats(0)
 	{}
+
+	FTimeSynthTimeDef(int32 InNumBars, int32 InNumBeats)
+		: NumBars(InNumBars)
+		, NumBeats(InNumBeats)
+	{
+	}
 };
 
 // Struct used to define a handle to a clip
diff --git a/Engine/Plugins/Runtime/TimeSynth/Source/TimeSynth/Private/TimeSynthComponent.cpp b/Engine/Plugins/Runtime/TimeSynth/Source/TimeSynth/Private/TimeSynthComponent.cpp
--- a/Engine/Plugins/Runtime/TimeSynth/Source/TimeSynth/Private/TimeSynthComponent.cpp
+++ b/Engine/Plugins/Runtime/TimeSynth/Source/TimeSynth/Private/TimeSynthComponent.cpp
@@ -736,7 +736,10 @@ FTimeSynthClipHandle UTimeSynthComponent::PlayClip(UTimeSynthClip* InClip, UTime
 
 	FTimeSynthTimeDef ClipDuration = InClip->ClipDuration;
 	FTimeSynthTimeDef FadeInTime = InClip->FadeInTime;
-	FTimeSynthTimeDef FadeOutTime = InClip->FadeOutTime;
+
+	FTimeSynthTimeDef FadeOutTime = InClip->bApplyFadeOut
+		? InClip->FadeOutTime
+		: FTimeSynthTimeDef(0, 0);
 
 	// Send this new clip over to the audio render thread
 	SynthCommand([this, NewClipInfo, ClipDuration, FadeInTime, FadeOutTime]
