Don't play clip of 0 duration & fade times

diff --git a/Engine/Plugins/Runtime/TimeSynth/Source/TimeSynth/Classes/TimeSynthComponent.h b/Engine/Plugins/Runtime/TimeSynth/Source/TimeSynth/Classes/TimeSynthComponent.h
--- a/Engine/Plugins/Runtime/TimeSynth/Source/TimeSynth/Classes/TimeSynthComponent.h
+++ b/Engine/Plugins/Runtime/TimeSynth/Source/TimeSynth/Classes/TimeSynthComponent.h
@@ -185,6 +185,11 @@ struct TIMESYNTH_API FTimeSynthTimeDef
 		, NumBeats(InNumBeats)
 	{
 	}
+
+	bool IsZeroDuration() const
+	{
+		return NumBars == 0 && NumBeats == 0;
+	}
 };
 
 // Struct used to define a handle to a clip
diff --git a/Engine/Plugins/Runtime/TimeSynth/Source/TimeSynth/Private/TimeSynthComponent.cpp b/Engine/Plugins/Runtime/TimeSynth/Source/TimeSynth/Private/TimeSynthComponent.cpp
--- a/Engine/Plugins/Runtime/TimeSynth/Source/TimeSynth/Private/TimeSynthComponent.cpp
+++ b/Engine/Plugins/Runtime/TimeSynth/Source/TimeSynth/Private/TimeSynthComponent.cpp
@@ -576,6 +576,15 @@ FTimeSynthClipHandle UTimeSynthComponent::PlayClip(UTimeSynthClip* InClip, UTime
 		return FTimeSynthClipHandle();
 	}
 
+	const bool bNoFadeIn = InClip->FadeInTime.IsZeroDuration();
+	const bool bNoDuration = InClip->ClipDuration.IsZeroDuration();
+	const bool bNoFadeOut = !InClip->bApplyFadeOut || InClip->FadeOutTime.IsZeroDuration();
+	if (bNoFadeIn && bNoDuration && bNoFadeOut)
+	{
+		UE_LOG(LogTimeSynth, Warning, TEXT("Failed to play clip: no duration or fade in/out set."));
+		return FTimeSynthClipHandle();
+	}
+
 	if (!bIsActive)
 	{
 		SetActive(true);
