Fixed issue with TimeSynth PIE crash caused by playing more clips than the synth was initially allocated for.

diff --git a/Engine/Plugins/Runtime/TimeSynth/Source/TimeSynth/Classes/TimeSynthComponent.h b/Engine/Plugins/Runtime/TimeSynth/Source/TimeSynth/Classes/TimeSynthComponent.h
--- a/Engine/Plugins/Runtime/TimeSynth/Source/TimeSynth/Classes/TimeSynthComponent.h
+++ b/Engine/Plugins/Runtime/TimeSynth/Source/TimeSynth/Classes/TimeSynthComponent.h
@@ -646,7 +646,7 @@ private:
 
 	int32 CurrentPoolSize;
 
-	// Array of free indicies int he playing clip pool
+	// Array of free indices in the playing clip pool
 	TArray<int32> FreePlayingClipIndices_AudioRenderThread;
 	TArray<int32> ActivePlayingClipIndices_AudioRenderThread;
 	
diff --git a/Engine/Plugins/Runtime/TimeSynth/Source/TimeSynth/Private/TimeSynthComponent.cpp b/Engine/Plugins/Runtime/TimeSynth/Source/TimeSynth/Private/TimeSynthComponent.cpp
--- a/Engine/Plugins/Runtime/TimeSynth/Source/TimeSynth/Private/TimeSynthComponent.cpp
+++ b/Engine/Plugins/Runtime/TimeSynth/Source/TimeSynth/Private/TimeSynthComponent.cpp
@@ -750,9 +750,11 @@ FTimeSynthClipHandle UTimeSynthComponent::PlayClip(UTimeSynthClip* InClip, UTime
 		? InClip->FadeOutTime
 		: FTimeSynthTimeDef(0, 0);
 
+
 	// Send this new clip over to the audio render thread
 	SynthCommand([this, NewClipInfo, ClipDuration, FadeInTime, FadeOutTime]
 	{
+		bool bCurrentPoolSizeChanged = false;
 		// Immediately create a mapping for this clip id to a free clip slot
 		// It's possible that the clip might get state changes before it starts playing if
 		// we're playing a very long-duration quantization
@@ -764,15 +766,25 @@ FTimeSynthClipHandle UTimeSynthComponent::PlayClip(UTimeSynthClip* InClip, UTime
 		else
 		{
 			// Grow the pool size if we ran out of clips in the pool
-			CurrentPoolSize++;
-			FreePlayingClipIndices_AudioRenderThread.Add(CurrentPoolSize);
+			++CurrentPoolSize;
+			bCurrentPoolSizeChanged = true;
+			FreePlayingClipIndices_AudioRenderThread.Add(CurrentPoolSize - 1);
 			FreeClipIndex = FreePlayingClipIndices_AudioRenderThread.Pop(false);
 		}
 		check(FreeClipIndex >= 0);
 
 		// Copy over the clip info to the slot
-		check(FreeClipIndex < PlayingClipsPool_AudioRenderThread.Num());
-		PlayingClipsPool_AudioRenderThread[FreeClipIndex] = NewClipInfo;
+		if(bCurrentPoolSizeChanged)
+		{
+			PlayingClipsPool_AudioRenderThread.Add(NewClipInfo);
+
+			UE_LOG(LogTimeSynth, Warning, TEXT("Reallocating PlayingClipsPool to %i (which is a performance hit.) If this wasn't caused by a hitch, consider initializing Pool Size to a larger value.")
+				, PlayingClipsPool_AudioRenderThread.Num());
+		}
+		else
+		{
+			PlayingClipsPool_AudioRenderThread[FreeClipIndex] = NewClipInfo;
+		}
 
 		// Add a mapping of the clip handle id to the free index
 		// This will allow us to reference the playing clip from BP, etc.
