Convert AssetTypeActions to use new menu API
#rnx

diff --git a/Engine/Plugins/2D/Paper2D/Source/Paper2DEditor/Paper2DEditor.Build.cs b/Engine/Plugins/2D/Paper2D/Source/Paper2DEditor/Paper2DEditor.Build.cs
--- a/Engine/Plugins/2D/Paper2D/Source/Paper2DEditor/Paper2DEditor.Build.cs
+++ b/Engine/Plugins/2D/Paper2D/Source/Paper2DEditor/Paper2DEditor.Build.cs
@@ -31,7 +31,8 @@ public class Paper2DEditor : ModuleRules
 				"EditorWidgets",
 				"Projects",
 				"NavigationSystem",
-            });
+				"EditorMenus",
+			});
 
 		PrivateIncludePathModuleNames.AddRange(
 			new string[] {
diff --git a/Engine/Plugins/2D/Paper2D/Source/Paper2DEditor/Private/SpriteAssetTypeActions.cpp b/Engine/Plugins/2D/Paper2D/Source/Paper2DEditor/Private/SpriteAssetTypeActions.cpp
--- a/Engine/Plugins/2D/Paper2D/Source/Paper2DEditor/Private/SpriteAssetTypeActions.cpp
+++ b/Engine/Plugins/2D/Paper2D/Source/Paper2DEditor/Private/SpriteAssetTypeActions.cpp
@@ -2,6 +2,7 @@
 
 #include "SpriteAssetTypeActions.h"
 #include "Framework/MultiBox/MultiBoxBuilder.h"
+#include "EditorMenuSubsystem.h"
 #include "Misc/PackageName.h"
 #include "Misc/FeedbackContext.h"
 #include "EditorStyleSet.h"
@@ -58,11 +59,12 @@ uint32 FSpriteAssetTypeActions::GetCategories()
 	return MyAssetCategory;
 }
 
-void FSpriteAssetTypeActions::GetActions(const TArray<UObject*>& InObjects, FMenuBuilder& MenuBuilder)
+void FSpriteAssetTypeActions::GetActions(const TArray<UObject*>& InObjects, FEditorMenuSection& Section)
 {
 	auto Sprites = GetTypedWeakObjectPtrs<UPaperSprite>(InObjects);
 
-	MenuBuilder.AddMenuEntry(
+	Section.AddMenuEntry(
+		"Sprite_CreateFlipbook",
 		LOCTEXT("Sprite_CreateFlipbook", "Create Flipbook"),
 		LOCTEXT("Sprite_CreateFlipbookTooltip", "Creates flipbooks from the selected sprites."),
 		FSlateIcon(FEditorStyle::GetStyleSetName(), "ClassIcon.PaperFlipbook"),
diff --git a/Engine/Plugins/2D/Paper2D/Source/Paper2DEditor/Private/SpriteAssetTypeActions.h b/Engine/Plugins/2D/Paper2D/Source/Paper2DEditor/Private/SpriteAssetTypeActions.h
--- a/Engine/Plugins/2D/Paper2D/Source/Paper2DEditor/Private/SpriteAssetTypeActions.h
+++ b/Engine/Plugins/2D/Paper2D/Source/Paper2DEditor/Private/SpriteAssetTypeActions.h
@@ -6,7 +6,6 @@
 #include "Toolkits/IToolkitHost.h"
 #include "AssetTypeActions_Base.h"
 
-class FMenuBuilder;
 class UPaperSprite;
 
 class FSpriteAssetTypeActions : public FAssetTypeActions_Base
@@ -19,7 +18,7 @@ public:
 	virtual FColor GetTypeColor() const override;
 	virtual UClass* GetSupportedClass() const override;
 	virtual bool HasActions(const TArray<UObject*>& InObjects) const override { return true; }
-	virtual void GetActions(const TArray<UObject*>& InObjects, FMenuBuilder& MenuBuilder) override;
+	virtual void GetActions(const TArray<UObject*>& InObjects, struct FEditorMenuSection& Section) override;
 	virtual void OpenAssetEditor(const TArray<UObject*>& InObjects, TSharedPtr<class IToolkitHost> EditWithinLevelEditor = TSharedPtr<IToolkitHost>()) override;
 	virtual uint32 GetCategories() override;
 	// End of IAssetTypeActions interface
diff --git a/Engine/Plugins/2D/Paper2D/Source/Paper2DEditor/Private/TileSetAssetTypeActions.cpp b/Engine/Plugins/2D/Paper2D/Source/Paper2DEditor/Private/TileSetAssetTypeActions.cpp
--- a/Engine/Plugins/2D/Paper2D/Source/Paper2DEditor/Private/TileSetAssetTypeActions.cpp
+++ b/Engine/Plugins/2D/Paper2D/Source/Paper2DEditor/Private/TileSetAssetTypeActions.cpp
@@ -1,7 +1,7 @@
 // Copyright 1998-2019 Epic Games, Inc. All Rights Reserved.
 
 #include "TileSetAssetTypeActions.h"
-#include "Framework/MultiBox/MultiBoxBuilder.h"
+#include "EditorMenuSubsystem.h"
 #include "Engine/Texture.h"
 #include "Misc/PackageName.h"
 #include "EditorStyleSet.h"
@@ -58,20 +58,22 @@ uint32 FTileSetAssetTypeActions::GetCategories()
 	return MyAssetCategory;
 }
 
-void FTileSetAssetTypeActions::GetActions(const TArray<UObject*>& InObjects, FMenuBuilder& MenuBuilder)
+void FTileSetAssetTypeActions::GetActions(const TArray<UObject*>& InObjects, FEditorMenuSection& Section)
 {
 	TArray<TWeakObjectPtr<UPaperTileSet>> TileSets = GetTypedWeakObjectPtrs<UPaperTileSet>(InObjects);
 
 	if (TileSets.Num() == 1)
 	{
-		MenuBuilder.AddMenuEntry(
+		Section.AddMenuEntry(
+			"TileSet_CreateTileMap",
 			LOCTEXT("TileSet_CreateTileMap", "Create Tile Map"),
 			LOCTEXT("TileSet_CreateTileMapTooltip", "Creates a tile map using the selected tile set as a guide for tile size, etc..."),
 			FSlateIcon(FEditorStyle::GetStyleSetName(), "ClassIcon.PaperTileSet"),
 			FUIAction(FExecuteAction::CreateSP(this, &FTileSetAssetTypeActions::ExecuteCreateTileMap, TileSets[0]))
 			);
 
-		MenuBuilder.AddMenuEntry(
+		Section.AddMenuEntry(
+			"TileSet_ConditionTileSet",
 			LOCTEXT("TileSet_ConditionTileSet", "Condition Tile Sheet Texture"),
 			LOCTEXT("TileSet_ConditionTileSetTooltip", "Conditions the tile sheet texture for the selected tile set by duplicating tile edges to create a buffer zone around each tile"),
 			FSlateIcon(FEditorStyle::GetStyleSetName(), "ClassIcon.Texture2D"),
diff --git a/Engine/Plugins/2D/Paper2D/Source/Paper2DEditor/Private/TileSetAssetTypeActions.h b/Engine/Plugins/2D/Paper2D/Source/Paper2DEditor/Private/TileSetAssetTypeActions.h
--- a/Engine/Plugins/2D/Paper2D/Source/Paper2DEditor/Private/TileSetAssetTypeActions.h
+++ b/Engine/Plugins/2D/Paper2D/Source/Paper2DEditor/Private/TileSetAssetTypeActions.h
@@ -6,7 +6,6 @@
 #include "Toolkits/IToolkitHost.h"
 #include "AssetTypeActions_Base.h"
 
-class FMenuBuilder;
 class UPaperTileSet;
 
 class FTileSetAssetTypeActions : public FAssetTypeActions_Base
@@ -21,7 +20,7 @@ public:
 	virtual void OpenAssetEditor(const TArray<UObject*>& InObjects, TSharedPtr<class IToolkitHost> EditWithinLevelEditor = TSharedPtr<IToolkitHost>()) override;
 	virtual uint32 GetCategories() override;
 	virtual bool HasActions(const TArray<UObject*>& InObjects) const override { return true; }
-	virtual void GetActions(const TArray<UObject*>& InObjects, FMenuBuilder& MenuBuilder) override;
+	virtual void GetActions(const TArray<UObject*>& InObjects, struct FEditorMenuSection& Section) override;
 	// End of IAssetTypeActions interface
 
 private:
diff --git a/Engine/Plugins/2D/Paper2D/Source/PaperSpriteSheetImporter/PaperSpriteSheetImporter.Build.cs b/Engine/Plugins/2D/Paper2D/Source/PaperSpriteSheetImporter/PaperSpriteSheetImporter.Build.cs
--- a/Engine/Plugins/2D/Paper2D/Source/PaperSpriteSheetImporter/PaperSpriteSheetImporter.Build.cs
+++ b/Engine/Plugins/2D/Paper2D/Source/PaperSpriteSheetImporter/PaperSpriteSheetImporter.Build.cs
@@ -19,7 +19,8 @@ public class PaperSpriteSheetImporter : ModuleRules
 				"Paper2DEditor",
 				"AssetTools",
 				"ContentBrowser",
-                "EditorStyle"
+                "EditorStyle",
+				"EditorMenus",
 			});
 
 		PrivateIncludePathModuleNames.AddRange(
diff --git a/Engine/Plugins/2D/Paper2D/Source/PaperSpriteSheetImporter/Private/PaperSpriteSheetAssetTypeActions.cpp b/Engine/Plugins/2D/Paper2D/Source/PaperSpriteSheetImporter/Private/PaperSpriteSheetAssetTypeActions.cpp
--- a/Engine/Plugins/2D/Paper2D/Source/PaperSpriteSheetImporter/Private/PaperSpriteSheetAssetTypeActions.cpp
+++ b/Engine/Plugins/2D/Paper2D/Source/PaperSpriteSheetImporter/Private/PaperSpriteSheetAssetTypeActions.cpp
@@ -2,7 +2,7 @@
 
 #include "PaperSpriteSheetAssetTypeActions.h"
 #include "Misc/PackageName.h"
-#include "Framework/MultiBox/MultiBoxBuilder.h"
+#include "EditorMenuSubsystem.h"
 #include "Misc/FeedbackContext.h"
 #include "EditorStyleSet.h"
 #include "EditorFramework/AssetImportData.h"
@@ -57,11 +57,12 @@ void FPaperSpriteSheetAssetTypeActions::GetResolvedSourceFilePaths(const TArray<
 	}
 }
 
-void FPaperSpriteSheetAssetTypeActions::GetActions(const TArray<UObject*>& InObjects, FMenuBuilder& MenuBuilder)
+void FPaperSpriteSheetAssetTypeActions::GetActions(const TArray<UObject*>& InObjects, FEditorMenuSection& Section)
 {
 	auto SpriteSheetImports = GetTypedWeakObjectPtrs<UPaperSpriteSheet>(InObjects);
 
-	MenuBuilder.AddMenuEntry(
+	Section.AddMenuEntry(
+		"SpriteSheet_CreateFlipbooks",
 		LOCTEXT("SpriteSheet_CreateFlipbooks", "Create Flipbooks"),
 		LOCTEXT("SpriteSheet_CreateFlipbooksTooltip", "Creates flipbooks from sprites in this sprite sheet."),
 		FSlateIcon(FEditorStyle::GetStyleSetName(), "ClassIcon.PaperFlipbook"),
diff --git a/Engine/Plugins/2D/Paper2D/Source/PaperSpriteSheetImporter/Private/PaperSpriteSheetAssetTypeActions.h b/Engine/Plugins/2D/Paper2D/Source/PaperSpriteSheetImporter/Private/PaperSpriteSheetAssetTypeActions.h
--- a/Engine/Plugins/2D/Paper2D/Source/PaperSpriteSheetImporter/Private/PaperSpriteSheetAssetTypeActions.h
+++ b/Engine/Plugins/2D/Paper2D/Source/PaperSpriteSheetImporter/Private/PaperSpriteSheetAssetTypeActions.h
@@ -5,7 +5,6 @@
 #include "CoreMinimal.h"
 #include "AssetTypeActions_Base.h"
 
-class FMenuBuilder;
 class UPaperSpriteSheet;
 
 class FPaperSpriteSheetAssetTypeActions : public FAssetTypeActions_Base
@@ -16,7 +15,7 @@ public:
 	virtual FColor GetTypeColor() const override;
 	virtual UClass* GetSupportedClass() const override;
 	virtual bool HasActions(const TArray<UObject*>& InObjects) const override { return true; }
-	virtual void GetActions(const TArray<UObject*>& InObjects, FMenuBuilder& MenuBuilder) override;
+	virtual void GetActions(const TArray<UObject*>& InObjects, struct FEditorMenuSection& Section) override;
 	virtual uint32 GetCategories() override;
 	virtual bool IsImportedAsset() const override;
 	virtual void GetResolvedSourceFilePaths(const TArray<UObject*>& TypeAssets, TArray<FString>& OutSourceFilePaths) const override;
diff --git a/Engine/Plugins/Editor/GameplayTagsEditor/Source/GameplayTagsEditor/GameplayTagsEditor.Build.cs b/Engine/Plugins/Editor/GameplayTagsEditor/Source/GameplayTagsEditor/GameplayTagsEditor.Build.cs
--- a/Engine/Plugins/Editor/GameplayTagsEditor/Source/GameplayTagsEditor/GameplayTagsEditor.Build.cs
+++ b/Engine/Plugins/Editor/GameplayTagsEditor/Source/GameplayTagsEditor/GameplayTagsEditor.Build.cs
@@ -34,7 +34,8 @@ namespace UnrealBuildTool.Rules
 					"ContentBrowser",
 					"MainFrame",
 					"UnrealEd",
-					"SourceControl"
+					"SourceControl",
+					"EditorMenus",
 				}
 			);
 
diff --git a/Engine/Plugins/Editor/GameplayTagsEditor/Source/GameplayTagsEditor/Private/AssetTypeActions_GameplayTagAssetBase.cpp b/Engine/Plugins/Editor/GameplayTagsEditor/Source/GameplayTagsEditor/Private/AssetTypeActions_GameplayTagAssetBase.cpp
--- a/Engine/Plugins/Editor/GameplayTagsEditor/Source/GameplayTagsEditor/Private/AssetTypeActions_GameplayTagAssetBase.cpp
+++ b/Engine/Plugins/Editor/GameplayTagsEditor/Source/GameplayTagsEditor/Private/AssetTypeActions_GameplayTagAssetBase.cpp
@@ -1,7 +1,7 @@
 // Copyright 1998-2019 Epic Games, Inc. All Rights Reserved.
 
 #include "AssetTypeActions_GameplayTagAssetBase.h"
-#include "Framework/MultiBox/MultiBoxBuilder.h"
+#include "EditorMenuSubsystem.h"
 #include "UObject/UnrealType.h"
 #include "Framework/Application/SlateApplication.h"
 
@@ -20,7 +20,7 @@ bool FAssetTypeActions_GameplayTagAssetBase::HasActions(const TArray<UObject*>&
 	return true;
 }
 
-void FAssetTypeActions_GameplayTagAssetBase::GetActions(const TArray<UObject*>& InObjects, FMenuBuilder& MenuBuilder)
+void FAssetTypeActions_GameplayTagAssetBase::GetActions(const TArray<UObject*>& InObjects, FEditorMenuSection& Section)
 {
 	TArray<UObject*> ContainerObjectOwners;
 	TArray<FGameplayTagContainer*> Containers;
@@ -41,7 +41,8 @@ void FAssetTypeActions_GameplayTagAssetBase::GetActions(const TArray<UObject*>&
 	ensure(Containers.Num() == ContainerObjectOwners.Num());
 	if (Containers.Num() > 0 && (Containers.Num() == ContainerObjectOwners.Num()))
 	{
-		MenuBuilder.AddMenuEntry(
+		Section.AddMenuEntry(
+			"GameplayTags_Edit",
 			LOCTEXT("GameplayTags_Edit", "Edit Gameplay Tags..."),
 			LOCTEXT("GameplayTags_EditToolTip", "Opens the Gameplay Tag Editor."),
 			FSlateIcon(),
diff --git a/Engine/Plugins/Editor/GameplayTagsEditor/Source/GameplayTagsEditor/Public/AssetTypeActions_GameplayTagAssetBase.h b/Engine/Plugins/Editor/GameplayTagsEditor/Source/GameplayTagsEditor/Public/AssetTypeActions_GameplayTagAssetBase.h
--- a/Engine/Plugins/Editor/GameplayTagsEditor/Source/GameplayTagsEditor/Public/AssetTypeActions_GameplayTagAssetBase.h
+++ b/Engine/Plugins/Editor/GameplayTagsEditor/Source/GameplayTagsEditor/Public/AssetTypeActions_GameplayTagAssetBase.h
@@ -5,7 +5,6 @@
 #include "CoreMinimal.h"
 #include "AssetTypeActions_Base.h"
 
-class FMenuBuilder;
 
 /** Base asset type actions for any classes with gameplay tagging */
 class GAMEPLAYTAGSEDITOR_API FAssetTypeActions_GameplayTagAssetBase : public FAssetTypeActions_Base
@@ -19,7 +18,7 @@ public:
 	virtual bool HasActions(const TArray<UObject*>& InObjects) const override;
 
 	/** Overridden to offer the gameplay tagging options */
-	virtual void GetActions(const TArray<UObject*>& InObjects, FMenuBuilder& MenuBuilder) override;
+	virtual void GetActions(const TArray<UObject*>& InObjects, struct FEditorMenuSection& Section) override;
 
 	/** Overridden to specify misc category */
 	virtual uint32 GetCategories() override;
diff --git a/Engine/Plugins/Experimental/ChaosSolverPlugin/Source/ChaosSolverEditor/ChaosSolverEditor.Build.cs b/Engine/Plugins/Experimental/ChaosSolverPlugin/Source/ChaosSolverEditor/ChaosSolverEditor.Build.cs
--- a/Engine/Plugins/Experimental/ChaosSolverPlugin/Source/ChaosSolverEditor/ChaosSolverEditor.Build.cs
+++ b/Engine/Plugins/Experimental/ChaosSolverPlugin/Source/ChaosSolverEditor/ChaosSolverEditor.Build.cs
@@ -25,8 +25,9 @@ namespace UnrealBuildTool.Rules
                     "ChaosSolverEngine",
                     "RawMesh",
                     "AssetTools",
-                    "AssetRegistry"
-                }
+                    "AssetRegistry",
+					"EditorMenus",
+				}
 				);
 		}
 	}
diff --git a/Engine/Plugins/Experimental/ChaosSolverPlugin/Source/ChaosSolverEditor/Private/Chaos/AssetTypeActions_ChaosSolver.cpp b/Engine/Plugins/Experimental/ChaosSolverPlugin/Source/ChaosSolverEditor/Private/Chaos/AssetTypeActions_ChaosSolver.cpp
--- a/Engine/Plugins/Experimental/ChaosSolverPlugin/Source/ChaosSolverEditor/Private/Chaos/AssetTypeActions_ChaosSolver.cpp
+++ b/Engine/Plugins/Experimental/ChaosSolverPlugin/Source/ChaosSolverEditor/Private/Chaos/AssetTypeActions_ChaosSolver.cpp
@@ -4,13 +4,13 @@
 
 #include "ThumbnailRendering/SceneThumbnailInfo.h"
 #include "Chaos/ChaosSolver.h"
+#include "EditorMenuSubsystem.h"
 
 #define LOCTEXT_NAMESPACE "AssetTypeActions"
 
 UClass* FAssetTypeActions_ChaosSolver::GetSupportedClass() const
 {
-	return 
-		UChaosSolver::StaticClass();
+	return UChaosSolver::StaticClass();
 }
 
 UThumbnailInfo* FAssetTypeActions_ChaosSolver::GetThumbnailInfo(UObject* Asset) const
@@ -19,9 +19,9 @@ UThumbnailInfo* FAssetTypeActions_ChaosSolver::GetThumbnailInfo(UObject* Asset)
 	return NewObject<USceneThumbnailInfo>(ChaosSolver, NAME_None, RF_Transactional);
 }
 
-void FAssetTypeActions_ChaosSolver::GetActions(const TArray<UObject*>& InObjects, FMenuBuilder& MenuBuilder)
+void FAssetTypeActions_ChaosSolver::GetActions(const TArray<UObject*>& InObjects, FEditorMenuSection& Section)
 {
-	FAssetTypeActions_Base::GetActions(InObjects, MenuBuilder);
+	FAssetTypeActions_Base::GetActions(InObjects, Section);
 }
 
 void FAssetTypeActions_ChaosSolver::OpenAssetEditor(const TArray<UObject*>& InObjects, TSharedPtr<IToolkitHost> EditWithinLevelEditor)
diff --git a/Engine/Plugins/Experimental/ChaosSolverPlugin/Source/ChaosSolverEditor/Public/Chaos/AssetTypeActions_ChaosSolver.h b/Engine/Plugins/Experimental/ChaosSolverPlugin/Source/ChaosSolverEditor/Public/Chaos/AssetTypeActions_ChaosSolver.h
--- a/Engine/Plugins/Experimental/ChaosSolverPlugin/Source/ChaosSolverEditor/Public/Chaos/AssetTypeActions_ChaosSolver.h
+++ b/Engine/Plugins/Experimental/ChaosSolverPlugin/Source/ChaosSolverEditor/Public/Chaos/AssetTypeActions_ChaosSolver.h
@@ -11,7 +11,7 @@ class FAssetTypeActions_ChaosSolver : public FAssetTypeActions_Base
 public:
 	// IAssetTypeActions Implementation
 	virtual FText GetName() const override { return NSLOCTEXT("AssetTypeActions", "AssetTypeActions_ChaosSolver", "Chaos Solver"); }
-	virtual void GetActions(const TArray<UObject*>& InObjects, FMenuBuilder& MenuBuilder) override;
+	virtual void GetActions(const TArray<UObject*>& InObjects, struct FEditorMenuSection& Section) override;
 	virtual FColor GetTypeColor() const override { return FColor(255, 192, 128); }
 	virtual UClass* GetSupportedClass() const override;
 	virtual void OpenAssetEditor(const TArray<UObject*>& InObjects, TSharedPtr<class IToolkitHost> EditWithinLevelEditor = TSharedPtr<IToolkitHost>()) override;
diff --git a/Engine/Plugins/Experimental/FieldSystemPlugin/Source/FieldSyStemEditor/FieldSystemEditor.Build.cs b/Engine/Plugins/Experimental/FieldSystemPlugin/Source/FieldSyStemEditor/FieldSystemEditor.Build.cs
--- a/Engine/Plugins/Experimental/FieldSystemPlugin/Source/FieldSyStemEditor/FieldSystemEditor.Build.cs
+++ b/Engine/Plugins/Experimental/FieldSystemPlugin/Source/FieldSyStemEditor/FieldSystemEditor.Build.cs
@@ -24,8 +24,9 @@ namespace UnrealBuildTool.Rules
                     "FieldSystemEngine",
                     "RawMesh",
                     "AssetTools",
-                    "AssetRegistry"
-                }
+                    "AssetRegistry",
+					"EditorMenus",
+				}
 				);
 		}
 	}
diff --git a/Engine/Plugins/Experimental/FieldSystemPlugin/Source/FieldSyStemEditor/Private/Field/AssetTypeActions_FieldSystem.cpp b/Engine/Plugins/Experimental/FieldSystemPlugin/Source/FieldSyStemEditor/Private/Field/AssetTypeActions_FieldSystem.cpp
--- a/Engine/Plugins/Experimental/FieldSystemPlugin/Source/FieldSyStemEditor/Private/Field/AssetTypeActions_FieldSystem.cpp
+++ b/Engine/Plugins/Experimental/FieldSystemPlugin/Source/FieldSyStemEditor/Private/Field/AssetTypeActions_FieldSystem.cpp
@@ -5,13 +5,13 @@
 #include "ThumbnailRendering/SceneThumbnailInfo.h"
 #include "Field/FieldSystem.h"
 #include "Field/FieldSystemAsset.h"
+#include "EditorMenuSubsystem.h"
 
 #define LOCTEXT_NAMESPACE "AssetTypeActions"
 
 UClass* FAssetTypeActions_FieldSystem::GetSupportedClass() const
 {
-	return 
-		UFieldSystem::StaticClass();
+	return UFieldSystem::StaticClass();
 }
 
 UThumbnailInfo* FAssetTypeActions_FieldSystem::GetThumbnailInfo(UObject* Asset) const
@@ -20,9 +20,9 @@ UThumbnailInfo* FAssetTypeActions_FieldSystem::GetThumbnailInfo(UObject* Asset)
 	return NewObject<USceneThumbnailInfo>(FieldSystem, NAME_None, RF_Transactional);
 }
 
-void FAssetTypeActions_FieldSystem::GetActions(const TArray<UObject*>& InObjects, FMenuBuilder& MenuBuilder)
+void FAssetTypeActions_FieldSystem::GetActions(const TArray<UObject*>& InObjects, FEditorMenuSection& Section)
 {
-	FAssetTypeActions_Base::GetActions(InObjects, MenuBuilder);
+	FAssetTypeActions_Base::GetActions(InObjects, Section);
 }
 
 void FAssetTypeActions_FieldSystem::OpenAssetEditor(const TArray<UObject*>& InObjects, TSharedPtr<IToolkitHost> EditWithinLevelEditor)
diff --git a/Engine/Plugins/Experimental/FieldSystemPlugin/Source/FieldSyStemEditor/Public/Field/AssetTypeActions_FieldSystem.h b/Engine/Plugins/Experimental/FieldSystemPlugin/Source/FieldSyStemEditor/Public/Field/AssetTypeActions_FieldSystem.h
--- a/Engine/Plugins/Experimental/FieldSystemPlugin/Source/FieldSyStemEditor/Public/Field/AssetTypeActions_FieldSystem.h
+++ b/Engine/Plugins/Experimental/FieldSystemPlugin/Source/FieldSyStemEditor/Public/Field/AssetTypeActions_FieldSystem.h
@@ -11,7 +11,7 @@ class FAssetTypeActions_FieldSystem : public FAssetTypeActions_Base
 public:
 	// IAssetTypeActions Implementation
 	virtual FText GetName() const override { return NSLOCTEXT("AssetTypeActions", "AssetTypeActions_FieldSystem", "Field System"); }
-	virtual void GetActions(const TArray<UObject*>& InObjects, FMenuBuilder& MenuBuilder) override;
+	virtual void GetActions(const TArray<UObject*>& InObjects, struct FEditorMenuSection& Section) override;
 	virtual FColor GetTypeColor() const override { return FColor(255, 192, 128); }
 	virtual UClass* GetSupportedClass() const override;
 	virtual void OpenAssetEditor(const TArray<UObject*>& InObjects, TSharedPtr<class IToolkitHost> EditWithinLevelEditor = TSharedPtr<IToolkitHost>()) override;
diff --git a/Engine/Plugins/Experimental/GeometryCache/Source/GeometryCacheEd/Classes/AssetTypeActions_GeometryCache.h b/Engine/Plugins/Experimental/GeometryCache/Source/GeometryCacheEd/Classes/AssetTypeActions_GeometryCache.h
--- a/Engine/Plugins/Experimental/GeometryCache/Source/GeometryCacheEd/Classes/AssetTypeActions_GeometryCache.h
+++ b/Engine/Plugins/Experimental/GeometryCache/Source/GeometryCacheEd/Classes/AssetTypeActions_GeometryCache.h
@@ -5,7 +5,6 @@
 #include "Toolkits/IToolkitHost.h"
 #include "AssetTypeActions_Base.h"
 
-class FMenuBuilder;
 
 /** Type-actions class for GeometryCache assets*/
 class FAssetTypeActions_GeometryCache : public FAssetTypeActions_Base
@@ -16,7 +15,7 @@ public:
 	virtual FColor GetTypeColor() const override;
 	virtual UClass* GetSupportedClass() const override;
 	virtual bool HasActions(const TArray<UObject*>& InObjects) const override;
-	virtual void GetActions(const TArray<UObject*>& InObjects, FMenuBuilder& MenuBuilder) override;
+	virtual void GetActions(const TArray<UObject*>& InObjects, struct FEditorMenuSection& Section) override;
 	virtual void OpenAssetEditor(const TArray<UObject*>& InObjects, TSharedPtr<class IToolkitHost> EditWithinLevelEditor = TSharedPtr<IToolkitHost>()) override;
 	virtual uint32 GetCategories() override;
 	virtual class UThumbnailInfo* GetThumbnailInfo(UObject* Asset) const override;
diff --git a/Engine/Plugins/Experimental/GeometryCache/Source/GeometryCacheEd/GeometryCacheEd.Build.cs b/Engine/Plugins/Experimental/GeometryCache/Source/GeometryCacheEd/GeometryCacheEd.Build.cs
--- a/Engine/Plugins/Experimental/GeometryCache/Source/GeometryCacheEd/GeometryCacheEd.Build.cs
+++ b/Engine/Plugins/Experimental/GeometryCache/Source/GeometryCacheEd/GeometryCacheEd.Build.cs
@@ -18,7 +18,8 @@ public class GeometryCacheEd : ModuleRules
                 "RHI",		
                 "UnrealEd",
 				"AssetTools",
-                "GeometryCache"
+                "GeometryCache",
+				"EditorMenus",
 			}
 		);
 	}
diff --git a/Engine/Plugins/Experimental/GeometryCache/Source/GeometryCacheEd/Private/AssetTypeActions_GeometryCache.cpp b/Engine/Plugins/Experimental/GeometryCache/Source/GeometryCacheEd/Private/AssetTypeActions_GeometryCache.cpp
--- a/Engine/Plugins/Experimental/GeometryCache/Source/GeometryCacheEd/Private/AssetTypeActions_GeometryCache.cpp
+++ b/Engine/Plugins/Experimental/GeometryCache/Source/GeometryCacheEd/Private/AssetTypeActions_GeometryCache.cpp
@@ -4,6 +4,7 @@
 #include "EditorFramework/AssetImportData.h"
 #include "ThumbnailRendering/SceneThumbnailInfo.h"
 #include "GeometryCache.h"
+#include "EditorMenuSubsystem.h"
 
 FText FAssetTypeActions_GeometryCache::GetName() const
 {
@@ -25,9 +26,9 @@ bool FAssetTypeActions_GeometryCache::HasActions(const TArray<UObject*>& InObjec
 	return false;
 }
 
-void FAssetTypeActions_GeometryCache::GetActions(const TArray<UObject*>& InObjects, FMenuBuilder& MenuBuilder)
+void FAssetTypeActions_GeometryCache::GetActions(const TArray<UObject*>& InObjects, FEditorMenuSection& Section)
 {
-	FAssetTypeActions_Base::GetActions(InObjects, MenuBuilder);
+	FAssetTypeActions_Base::GetActions(InObjects, Section);
 }
 
 void FAssetTypeActions_GeometryCache::OpenAssetEditor(const TArray<UObject*>& InObjects, TSharedPtr<class IToolkitHost> EditWithinLevelEditor /*= TSharedPtr<IToolkitHost>()*/)
diff --git a/Engine/Plugins/Experimental/GeometryCollectionPlugin/Source/GeometryCollectionEditor/GeometryCollectionEditor.Build.cs b/Engine/Plugins/Experimental/GeometryCollectionPlugin/Source/GeometryCollectionEditor/GeometryCollectionEditor.Build.cs
--- a/Engine/Plugins/Experimental/GeometryCollectionPlugin/Source/GeometryCollectionEditor/GeometryCollectionEditor.Build.cs
+++ b/Engine/Plugins/Experimental/GeometryCollectionPlugin/Source/GeometryCollectionEditor/GeometryCollectionEditor.Build.cs
@@ -29,7 +29,8 @@ namespace UnrealBuildTool.Rules
 				    "AssetRegistry",
 				    "SceneOutliner",
 				    "FieldSystemCore",
-					"EditorStyle"
+					"EditorStyle",
+					"EditorMenus",
 				}
 			);
 		}
diff --git a/Engine/Plugins/Experimental/GeometryCollectionPlugin/Source/GeometryCollectionEditor/Private/GeometryCollection/AssetTypeActions_GeometryCollection.cpp b/Engine/Plugins/Experimental/GeometryCollectionPlugin/Source/GeometryCollectionEditor/Private/GeometryCollection/AssetTypeActions_GeometryCollection.cpp
--- a/Engine/Plugins/Experimental/GeometryCollectionPlugin/Source/GeometryCollectionEditor/Private/GeometryCollection/AssetTypeActions_GeometryCollection.cpp
+++ b/Engine/Plugins/Experimental/GeometryCollectionPlugin/Source/GeometryCollectionEditor/Private/GeometryCollection/AssetTypeActions_GeometryCollection.cpp
@@ -3,13 +3,13 @@
 #include "GeometryCollection/AssetTypeActions_GeometryCollection.h"
 #include "ThumbnailRendering/SceneThumbnailInfo.h"
 #include "GeometryCollection/GeometryCollectionObject.h"
+#include "EditorMenuSubsystem.h"
 
 #define LOCTEXT_NAMESPACE "AssetTypeActions"
 
 UClass* FAssetTypeActions_GeometryCollection::GetSupportedClass() const
 {
-	return 
-		UGeometryCollection::StaticClass();
+	return UGeometryCollection::StaticClass();
 }
 
 UThumbnailInfo* FAssetTypeActions_GeometryCollection::GetThumbnailInfo(UObject* Asset) const
@@ -25,9 +25,9 @@ UThumbnailInfo* FAssetTypeActions_GeometryCollection::GetThumbnailInfo(UObject*
 	return ThumbnailInfo;
 }
 
-void FAssetTypeActions_GeometryCollection::GetActions(const TArray<UObject*>& InObjects, FMenuBuilder& MenuBuilder)
+void FAssetTypeActions_GeometryCollection::GetActions(const TArray<UObject*>& InObjects, FEditorMenuSection& Section)
 {
-	FAssetTypeActions_Base::GetActions(InObjects, MenuBuilder);
+	FAssetTypeActions_Base::GetActions(InObjects, Section);
 	//Set("ClassIcon.GeometryCollection", new IMAGE_BRUSH("Icons/AssetIcons/GeometryCollectionComponentAtlasGroup_16x", Icon16x16));
 	//Set("ClassThumbnail.GeometryCollection", new IMAGE_BRUSH("Icons/AssetIcons/GeometryCollectionComponentAtlasGroup_64x", Icon64x64));
 	// IconPath = Plugin->GetBaseDir() / TEXT("Resources/Icon128.png");
diff --git a/Engine/Plugins/Experimental/GeometryCollectionPlugin/Source/GeometryCollectionEditor/Public/GeometryCollection/AssetTypeActions_GeometryCollection.h b/Engine/Plugins/Experimental/GeometryCollectionPlugin/Source/GeometryCollectionEditor/Public/GeometryCollection/AssetTypeActions_GeometryCollection.h
--- a/Engine/Plugins/Experimental/GeometryCollectionPlugin/Source/GeometryCollectionEditor/Public/GeometryCollection/AssetTypeActions_GeometryCollection.h
+++ b/Engine/Plugins/Experimental/GeometryCollectionPlugin/Source/GeometryCollectionEditor/Public/GeometryCollection/AssetTypeActions_GeometryCollection.h
@@ -11,7 +11,7 @@ class FAssetTypeActions_GeometryCollection : public FAssetTypeActions_Base
 public:
 	// IAssetTypeActions Implementation
 	virtual FText GetName() const override { return NSLOCTEXT("AssetTypeActions", "AssetTypeActions_GeometryCollection", "Geometry Collection"); }
-	virtual void GetActions(const TArray<UObject*>& InObjects, FMenuBuilder& MenuBuilder) override;
+	virtual void GetActions(const TArray<UObject*>& InObjects, struct FEditorMenuSection& Section) override;
 	virtual FColor GetTypeColor() const override { return FColor(255, 192, 128); }
 	virtual UClass* GetSupportedClass() const override;
 	//virtual void OpenAssetEditor(const TArray<UObject*>& InObjects, TSharedPtr<class IToolkitHost> EditWithinLevelEditor = TSharedPtr<IToolkitHost>()) override;
diff --git a/Engine/Plugins/FX/HoudiniNiagara/Source/HoudiniNiagaraEditor/HoudiniNiagaraEditor.Build.cs b/Engine/Plugins/FX/HoudiniNiagara/Source/HoudiniNiagaraEditor/HoudiniNiagaraEditor.Build.cs
--- a/Engine/Plugins/FX/HoudiniNiagara/Source/HoudiniNiagaraEditor/HoudiniNiagaraEditor.Build.cs
+++ b/Engine/Plugins/FX/HoudiniNiagara/Source/HoudiniNiagaraEditor/HoudiniNiagaraEditor.Build.cs
@@ -53,7 +53,8 @@ public class HoudiniNiagaraEditor : ModuleRules
                 "HoudiniNiagara",
                 "UnrealEd",
 				"EditorStyle",
-            }
+				"EditorMenus",
+			}
         );
 			
 		
@@ -67,7 +68,8 @@ public class HoudiniNiagaraEditor : ModuleRules
                 "Niagara",
                 "NiagaraShader",
                 "UtilityShaders",
-                "HoudiniNiagara"
+                "HoudiniNiagara",
+				"EditorMenus",
 			}
         );
 
diff --git a/Engine/Plugins/FX/HoudiniNiagara/Source/HoudiniNiagaraEditor/Private/HoudiniCSVAssetActions.cpp b/Engine/Plugins/FX/HoudiniNiagara/Source/HoudiniNiagaraEditor/Private/HoudiniCSVAssetActions.cpp
--- a/Engine/Plugins/FX/HoudiniNiagara/Source/HoudiniNiagaraEditor/Private/HoudiniCSVAssetActions.cpp
+++ b/Engine/Plugins/FX/HoudiniNiagara/Source/HoudiniNiagaraEditor/Private/HoudiniCSVAssetActions.cpp
@@ -22,7 +22,7 @@
 */
 #include "HoudiniCSVAssetActions.h"
 
-#include "Framework/MultiBox/MultiBoxBuilder.h"
+#include "EditorMenuSubsystem.h"
 #include "HoudiniCSV.h"
 #include "EditorStyleSet.h"
 #include "Toolkits/AssetEditorToolkit.h"
@@ -43,13 +43,14 @@ bool FHoudiniCSVAssetActions::CanFilter()
 }
 
 
-void FHoudiniCSVAssetActions::GetActions(const TArray<UObject*>& InObjects, FMenuBuilder& MenuBuilder)
+void FHoudiniCSVAssetActions::GetActions(const TArray<UObject*>& InObjects, FEditorMenuSection& Section)
 {
-    FAssetTypeActions_Base::GetActions(InObjects, MenuBuilder);
+    FAssetTypeActions_Base::GetActions(InObjects, Section);
 
     auto HoudiniCSVAssets = GetTypedWeakObjectPtrs<UHoudiniCSV>(InObjects);
 
-    MenuBuilder.AddMenuEntry(
+    Section.AddMenuEntry(
+		"ReimportHoudiniCSVLabel",
 		LOCTEXT("ReimportHoudiniCSVLabel", "Reimport"),
 		LOCTEXT("ReimportHoudiniCSVTooltip", "Reimport the selected Houdini CSV file(s)."),
 		FSlateIcon(FEditorStyle::GetStyleSetName(), "ContentBrowser.AssetActions.ReimportAsset"),
@@ -59,7 +60,8 @@ void FHoudiniCSVAssetActions::GetActions(const TArray<UObject*>& InObjects, FMen
 		)
     );
 
-    MenuBuilder.AddMenuEntry(
+    Section.AddMenuEntry(
+		"OpenHoudiniCSVLabel",
 		LOCTEXT("OpenHoudiniCSVLabel", "Open in Text Editor"),
 		LOCTEXT("OpenHoudiniCSVTooltip", "Open the selected Houdini CSV file(s) in a Text Editor."),
 		FSlateIcon(FEditorStyle::GetStyleSetName(), "ContentBrowser.AssetActions.OpenInExternalEditor"),
@@ -69,7 +71,8 @@ void FHoudiniCSVAssetActions::GetActions(const TArray<UObject*>& InObjects, FMen
 		)
     );
 
-	MenuBuilder.AddMenuEntry(
+	Section.AddMenuEntry(
+		"FindHoudiniCSVInExplorer",
 		LOCTEXT("FindHoudiniCSVInExplorer", "Find Source CSV"),
 		LOCTEXT("FindHoudiniCSVInExplorerTooltip", "Opens explorer at the location of this asset's source CSV file."),
 		FSlateIcon(FEditorStyle::GetStyleSetName(), "ContentBrowser.AssetActions.OpenInExternalEditor"),
@@ -247,4 +250,4 @@ FHoudiniCSVAssetActions::ExecuteFindInExplorer(const TArray<TWeakObjectPtr<UHoud
 	}
 }
 
-#undef LOCTEXT_NAMESPACE
\ No newline at end of file
+#undef LOCTEXT_NAMESPACE
diff --git a/Engine/Plugins/FX/HoudiniNiagara/Source/HoudiniNiagaraEditor/Public/HoudiniCSVAssetActions.h b/Engine/Plugins/FX/HoudiniNiagara/Source/HoudiniNiagaraEditor/Public/HoudiniCSVAssetActions.h
--- a/Engine/Plugins/FX/HoudiniNiagara/Source/HoudiniNiagaraEditor/Public/HoudiniCSVAssetActions.h
+++ b/Engine/Plugins/FX/HoudiniNiagara/Source/HoudiniNiagaraEditor/Public/HoudiniCSVAssetActions.h
@@ -44,7 +44,7 @@ class FHoudiniCSVAssetActions : public FAssetTypeActions_Base
 	virtual uint32 GetCategories() override;    
 	virtual class UThumbnailInfo* GetThumbnailInfo(UObject* Asset) const override;
 	virtual bool HasActions(const TArray<UObject*>& InObjects) const override;
-	virtual void GetActions(const TArray<UObject*>& InObjects, FMenuBuilder& MenuBuilder) override; 
+	virtual void GetActions(const TArray<UObject*>& InObjects, struct FEditorMenuSection& Section) override;
 
 	virtual bool CanFilter() override;   
 	virtual FText GetAssetDescription(const FAssetData& AssetData) const override;
@@ -62,4 +62,4 @@ class FHoudiniCSVAssetActions : public FAssetTypeActions_Base
 	// Find In Explorer Actions
 	bool CanExecuteFindInExplorer(const TArray<TWeakObjectPtr<UHoudiniCSV>> Objects) const;
 	void ExecuteFindInExplorer(const TArray<TWeakObjectPtr<UHoudiniCSV>> Objects) const;	
-};
\ No newline at end of file
+};
diff --git a/Engine/Plugins/FX/Niagara/Source/NiagaraEditor/NiagaraEditor.Build.cs b/Engine/Plugins/FX/Niagara/Source/NiagaraEditor/NiagaraEditor.Build.cs
--- a/Engine/Plugins/FX/Niagara/Source/NiagaraEditor/NiagaraEditor.Build.cs
+++ b/Engine/Plugins/FX/Niagara/Source/NiagaraEditor/NiagaraEditor.Build.cs
@@ -47,6 +47,7 @@ public class NiagaraEditor : ModuleRules
                 "AdvancedPreviewScene",
 				"Projects",
                 "MainFrame",
+				"EditorMenus",
 			}
         );
 
diff --git a/Engine/Plugins/FX/Niagara/Source/NiagaraEditor/Private/AssetTypeActions/AssetTypeActions_NiagaraEmitter.cpp b/Engine/Plugins/FX/Niagara/Source/NiagaraEditor/Private/AssetTypeActions/AssetTypeActions_NiagaraEmitter.cpp
--- a/Engine/Plugins/FX/Niagara/Source/NiagaraEditor/Private/AssetTypeActions/AssetTypeActions_NiagaraEmitter.cpp
+++ b/Engine/Plugins/FX/Niagara/Source/NiagaraEditor/Private/AssetTypeActions/AssetTypeActions_NiagaraEmitter.cpp
@@ -8,7 +8,7 @@
 #include "NiagaraSystemFactoryNew.h"
 #include "ViewModels/NiagaraSystemViewModel.h"
 #include "NiagaraSystemScriptViewModel.h"
-#include "Framework/MultiBox/MultiBoxBuilder.h"
+#include "EditorMenuSubsystem.h"
 #include "EditorStyleSet.h"
 #include "Toolkits/AssetEditorToolkit.h"
 #include "IContentBrowserSingleton.h"
@@ -52,11 +52,12 @@ UClass* FAssetTypeActions_NiagaraEmitter::GetSupportedClass() const
 	return UNiagaraEmitter::StaticClass();
 }
 
-void FAssetTypeActions_NiagaraEmitter::GetActions(const TArray<UObject*>& InObjects, FMenuBuilder& MenuBuilder)
+void FAssetTypeActions_NiagaraEmitter::GetActions(const TArray<UObject*>& InObjects, FEditorMenuSection& Section)
 {
 	auto NiagaraEmitters = GetTypedWeakObjectPtrs<UNiagaraEmitter>(InObjects);
 
-	MenuBuilder.AddMenuEntry(
+	Section.AddMenuEntry(
+		"Emitter_NewNiagaraSystem",
 		LOCTEXT("Emitter_NewNiagaraSystem", "Create Niagara System"),
 		LOCTEXT("Emitter_NewNiagaraSystemTooltip", "Creates a niagara system using this emitter as a base."),
 		FSlateIcon(FEditorStyle::GetStyleSetName(), "ClassIcon.ParticleSystem"),
@@ -74,7 +75,8 @@ void FAssetTypeActions_NiagaraEmitter::GetActions(const TArray<UObject*>& InObje
 		}
 	}
 
-	MenuBuilder.AddMenuEntry(
+	Section.AddMenuEntry(
+		"MarkDependentCompilableAssetsDirty",
 		LOCTEXT("MarkDependentCompilableAssetsDirtyLabel", "Mark dependent compilable assets dirty"),
 		LOCTEXT("MarkDependentCompilableAssetsDirtyToolTip", "Finds all niagara assets which depend on this asset either directly or indirectly,\n and marks them dirty so they can be saved with the latest version."),
 		FSlateIcon(),
diff --git a/Engine/Plugins/FX/Niagara/Source/NiagaraEditor/Private/AssetTypeActions/AssetTypeActions_NiagaraEmitter.h b/Engine/Plugins/FX/Niagara/Source/NiagaraEditor/Private/AssetTypeActions/AssetTypeActions_NiagaraEmitter.h
--- a/Engine/Plugins/FX/Niagara/Source/NiagaraEditor/Private/AssetTypeActions/AssetTypeActions_NiagaraEmitter.h
+++ b/Engine/Plugins/FX/Niagara/Source/NiagaraEditor/Private/AssetTypeActions/AssetTypeActions_NiagaraEmitter.h
@@ -17,7 +17,7 @@ public:
 	virtual FColor GetTypeColor() const override;
 	virtual UClass* GetSupportedClass() const override;
 	virtual bool HasActions(const TArray<UObject*>& InObjects) const override { return true; }
-	virtual void GetActions( const TArray<UObject*>& InObjects, FMenuBuilder& MenuBuilder ) override;
+	virtual void GetActions(const TArray<UObject*>& InObjects, struct FEditorMenuSection& Section) override;
 	virtual void OpenAssetEditor(const TArray<UObject*>& InObjects, TSharedPtr<class IToolkitHost> EditWithinLevelEditor = TSharedPtr<IToolkitHost>()) override;
 	virtual uint32 GetCategories() override { return FNiagaraEditorModule::GetAssetCategory(); }
 
diff --git a/Engine/Plugins/FX/Niagara/Source/NiagaraEditor/Private/AssetTypeActions/AssetTypeActions_NiagaraParameterCollection.cpp b/Engine/Plugins/FX/Niagara/Source/NiagaraEditor/Private/AssetTypeActions/AssetTypeActions_NiagaraParameterCollection.cpp
--- a/Engine/Plugins/FX/Niagara/Source/NiagaraEditor/Private/AssetTypeActions/AssetTypeActions_NiagaraParameterCollection.cpp
+++ b/Engine/Plugins/FX/Niagara/Source/NiagaraEditor/Private/AssetTypeActions/AssetTypeActions_NiagaraParameterCollection.cpp
@@ -9,7 +9,7 @@
 #include "IContentBrowserSingleton.h"
 #include "ContentBrowserModule.h"
 
-#include "Framework/MultiBox/MultiBoxBuilder.h"
+#include "EditorMenuSubsystem.h"
 #include "EditorStyleSet.h"
 
 #define LOCTEXT_NAMESPACE "AssetTypeActions_NiagaraParameterCollection"
@@ -34,11 +34,12 @@ void FAssetTypeActions_NiagaraParameterCollection::OpenAssetEditor(const TArray<
 	}
 }
 
-void FAssetTypeActions_NiagaraParameterCollection::GetActions(const TArray<UObject*>& InObjects, FMenuBuilder& MenuBuilder)
+void FAssetTypeActions_NiagaraParameterCollection::GetActions(const TArray<UObject*>& InObjects, FEditorMenuSection& Section)
 {
 	auto Collections = GetTypedWeakObjectPtrs<UNiagaraParameterCollection>(InObjects);
 
-	MenuBuilder.AddMenuEntry(
+	Section.AddMenuEntry(
+		"Niagara_NewNPC",
 		LOCTEXT("NewNPC", "Create Niagara Parameter Collection Instance"),
 		LOCTEXT("NewNPCTooltip", "Creates an instance of this Niagara Parameter Collection."),
 		FSlateIcon(FEditorStyle::GetStyleSetName(), "ClassIcon.MaterialInstanceActor"),
diff --git a/Engine/Plugins/FX/Niagara/Source/NiagaraEditor/Private/AssetTypeActions/AssetTypeActions_NiagaraParameterCollection.h b/Engine/Plugins/FX/Niagara/Source/NiagaraEditor/Private/AssetTypeActions/AssetTypeActions_NiagaraParameterCollection.h
--- a/Engine/Plugins/FX/Niagara/Source/NiagaraEditor/Private/AssetTypeActions/AssetTypeActions_NiagaraParameterCollection.h
+++ b/Engine/Plugins/FX/Niagara/Source/NiagaraEditor/Private/AssetTypeActions/AssetTypeActions_NiagaraParameterCollection.h
@@ -17,7 +17,7 @@ public:
 	virtual void OpenAssetEditor(const TArray<UObject*>& InObjects, TSharedPtr<class IToolkitHost> EditWithinLevelEditor = TSharedPtr<IToolkitHost>()) override;
 	virtual uint32 GetCategories() override { return FNiagaraEditorModule::GetAssetCategory(); }
 	virtual bool HasActions(const TArray<UObject*>& InObjects) const override { return true; }
-	virtual void GetActions(const TArray<UObject*>& InObjects, FMenuBuilder& MenuBuilder)override;
+	virtual void GetActions(const TArray<UObject*>& InObjects, struct FEditorMenuSection& Section) override;
 
 	void ExecuteNewNPC(TArray<TWeakObjectPtr<class UNiagaraParameterCollection>> Objects);
 };
@@ -31,4 +31,4 @@ public:
 	virtual UClass* GetSupportedClass() const override;
 	virtual void OpenAssetEditor(const TArray<UObject*>& InObjects, TSharedPtr<class IToolkitHost> EditWithinLevelEditor = TSharedPtr<IToolkitHost>()) override;
 	virtual uint32 GetCategories() override { return FNiagaraEditorModule::GetAssetCategory(); }
-};
\ No newline at end of file
+};
diff --git a/Engine/Plugins/FX/Niagara/Source/NiagaraEditor/Private/AssetTypeActions/AssetTypeActions_NiagaraScript.cpp b/Engine/Plugins/FX/Niagara/Source/NiagaraEditor/Private/AssetTypeActions/AssetTypeActions_NiagaraScript.cpp
--- a/Engine/Plugins/FX/Niagara/Source/NiagaraEditor/Private/AssetTypeActions/AssetTypeActions_NiagaraScript.cpp
+++ b/Engine/Plugins/FX/Niagara/Source/NiagaraEditor/Private/AssetTypeActions/AssetTypeActions_NiagaraScript.cpp
@@ -6,7 +6,7 @@
 #include "NiagaraEditorStyle.h"
 #include "AssetData.h"
 #include "NiagaraEditorUtilities.h"
-#include "Framework/MultiBox/MultiBoxBuilder.h"
+#include "EditorMenuSubsystem.h"
 
 #define LOCTEXT_NAMESPACE "NiagaraScriptAssetTypeActions"
 
@@ -74,9 +74,10 @@ bool FAssetTypeActions_NiagaraScript::HasActions(const TArray<UObject*>& InObjec
 	return true;
 }
 
-void FAssetTypeActions_NiagaraScript::GetActions(const TArray<UObject*>& InObjects, FMenuBuilder& MenuBuilder)
+void FAssetTypeActions_NiagaraScript::GetActions(const TArray<UObject*>& InObjects, FEditorMenuSection& Section)
 {
-	MenuBuilder.AddMenuEntry(
+	Section.AddMenuEntry(
+		"MarkDependentCompilableAssetsDirty",
 		LOCTEXT("MarkDependentCompilableAssetsDirtyLabel", "Mark dependent compilable assets dirty"),
 		LOCTEXT("MarkDependentCompilableAssetsDirtyToolTip", "Finds all niagara assets which depend on this asset either directly or indirectly,\n and marks them dirty so they can be saved with the latest version."),
 		FSlateIcon(),
diff --git a/Engine/Plugins/FX/Niagara/Source/NiagaraEditor/Private/AssetTypeActions/AssetTypeActions_NiagaraScript.h b/Engine/Plugins/FX/Niagara/Source/NiagaraEditor/Private/AssetTypeActions/AssetTypeActions_NiagaraScript.h
--- a/Engine/Plugins/FX/Niagara/Source/NiagaraEditor/Private/AssetTypeActions/AssetTypeActions_NiagaraScript.h
+++ b/Engine/Plugins/FX/Niagara/Source/NiagaraEditor/Private/AssetTypeActions/AssetTypeActions_NiagaraScript.h
@@ -18,7 +18,7 @@ public:
 	virtual uint32 GetCategories() override { return FNiagaraEditorModule::GetAssetCategory(); }
 	virtual FText GetDisplayNameFromAssetData(const FAssetData& AssetData) const override;
 	virtual bool HasActions(const TArray<UObject*>& InObjects) const override;
-	virtual void GetActions(const TArray<UObject*>& InObjects, FMenuBuilder& MenuBuilder) override;
+	virtual void GetActions(const TArray<UObject*>& InObjects, struct FEditorMenuSection& Section) override;
 	// End IAssetTypeActions Interface
 
 	/** Get the action type name. */
diff --git a/Engine/Plugins/MovieScene/LevelSequenceEditor/Source/LevelSequenceEditor/Private/AssetTools/LevelSequenceActions.cpp b/Engine/Plugins/MovieScene/LevelSequenceEditor/Source/LevelSequenceEditor/Private/AssetTools/LevelSequenceActions.cpp
--- a/Engine/Plugins/MovieScene/LevelSequenceEditor/Source/LevelSequenceEditor/Private/AssetTools/LevelSequenceActions.cpp
+++ b/Engine/Plugins/MovieScene/LevelSequenceEditor/Source/LevelSequenceEditor/Private/AssetTools/LevelSequenceActions.cpp
@@ -5,7 +5,6 @@
 #include "Engine/Engine.h"
 #include "LevelSequence.h"
 #include "LevelSequenceEditorToolkit.h"
-#include "LevelSequenceActionExtender.h"
 
 #define LOCTEXT_NAMESPACE "AssetTypeActions"
 
@@ -89,19 +88,4 @@ bool FLevelSequenceActions::ShouldForceWorldCentric()
 }
 
 
-bool FLevelSequenceActions::HasActions( const TArray<UObject*>& InObjects ) const
-{
-	return ActionExtenders.Num() > 0;
-}
-
-
-void FLevelSequenceActions::GetActions( const TArray<UObject*>& InObjects, FMenuBuilder& MenuBuilder )
-{
-	for (const TSharedRef<FLevelSequenceActionExtender>& Extender : ActionExtenders)
-	{
-		Extender->GetActions(InObjects, MenuBuilder);
-	}
-}
-
-
 #undef LOCTEXT_NAMESPACE
diff --git a/Engine/Plugins/MovieScene/LevelSequenceEditor/Source/LevelSequenceEditor/Private/AssetTools/LevelSequenceActions.h b/Engine/Plugins/MovieScene/LevelSequenceEditor/Source/LevelSequenceEditor/Private/AssetTools/LevelSequenceActions.h
--- a/Engine/Plugins/MovieScene/LevelSequenceEditor/Source/LevelSequenceEditor/Private/AssetTools/LevelSequenceActions.h
+++ b/Engine/Plugins/MovieScene/LevelSequenceEditor/Source/LevelSequenceEditor/Private/AssetTools/LevelSequenceActions.h
@@ -7,7 +7,6 @@
 #include "Toolkits/IToolkitHost.h"
 #include "AssetTypeActions_Base.h"
 
-struct FLevelSequenceActionExtender;
 
 /**
  * Implements actions for ULevelSequence assets.
@@ -24,22 +23,6 @@ public:
 	 */
 	FLevelSequenceActions(const TSharedRef<ISlateStyle>& InStyle);
 
-	/**
-	 * Register a new extender that can add new actions to level sequence assets
-	 */
-	void RegisterLevelSequenceActionExtender(const TSharedRef<FLevelSequenceActionExtender>& InExtender)
-	{
-		ActionExtenders.AddUnique(InExtender);
-	}
-
-	/**
-	 * Unregister a previously registered action extender
-	 */
-	void UnregisterLevelSequenceActionExtender(const TSharedRef<FLevelSequenceActionExtender>& InExtender)
-	{
-		ActionExtenders.Remove(InExtender);
-	}
-
 public:
 	
 	// IAssetTypeActions interface
@@ -51,14 +34,10 @@ public:
 	virtual void OpenAssetEditor(const TArray<UObject*>& InObjects, TSharedPtr<class IToolkitHost> EditWithinLevelEditor = TSharedPtr<IToolkitHost>()) override;
 	virtual bool ShouldForceWorldCentric() override;
 	virtual bool CanLocalize() const override { return false; }
-	virtual bool HasActions( const TArray<UObject*>& InObjects ) const override;
-	virtual void GetActions( const TArray<UObject*>& InObjects, FMenuBuilder& MenuBuilder ) override;
+	virtual bool HasActions(const TArray<UObject*>& InObjects) const { return false; }
 
 private:
 
 	/** Pointer to the style set to use for toolkits. */
 	TSharedRef<ISlateStyle> Style;
-
-	/** Array of registered action extenders. */
-	TArray<TSharedRef<FLevelSequenceActionExtender>> ActionExtenders;
 };
diff --git a/Engine/Plugins/MovieScene/LevelSequenceEditor/Source/LevelSequenceEditor/Private/LevelSequenceEditorModule.cpp b/Engine/Plugins/MovieScene/LevelSequenceEditor/Source/LevelSequenceEditor/Private/LevelSequenceEditorModule.cpp
--- a/Engine/Plugins/MovieScene/LevelSequenceEditor/Source/LevelSequenceEditor/Private/LevelSequenceEditorModule.cpp
+++ b/Engine/Plugins/MovieScene/LevelSequenceEditor/Source/LevelSequenceEditor/Private/LevelSequenceEditorModule.cpp
@@ -369,16 +369,6 @@ protected:
 		return OnMasterSequenceCreatedEvent;
 	}
 
-	virtual void RegisterLevelSequenceActionExtender(TSharedRef<FLevelSequenceActionExtender> InExtender) override
-	{
-		LevelSequenceTypeActions->RegisterLevelSequenceActionExtender(InExtender);
-	}
-
-	virtual void UnregisterLevelSequenceActionExtender(TSharedRef<FLevelSequenceActionExtender> InExtender) override
-	{
-		LevelSequenceTypeActions->UnregisterLevelSequenceActionExtender(InExtender);
-	}
-
 	static TSharedRef<ISequencerEditorObjectBinding> OnCreateActorBinding(TSharedRef<ISequencer> InSequencer)
 	{
 		return MakeShareable(new FLevelSequenceEditorActorBinding(InSequencer));
diff --git a/Engine/Plugins/MovieScene/LevelSequenceEditor/Source/LevelSequenceEditor/Public/LevelSequenceActionExtender.h b/Engine/Plugins/MovieScene/LevelSequenceEditor/Source/LevelSequenceEditor/Public/LevelSequenceActionExtender.h
deleted file mode 100644
diff --git a/Engine/Plugins/MovieScene/LevelSequenceEditor/Source/LevelSequenceEditor/Public/LevelSequenceEditorModule.h b/Engine/Plugins/MovieScene/LevelSequenceEditor/Source/LevelSequenceEditor/Public/LevelSequenceEditorModule.h
--- a/Engine/Plugins/MovieScene/LevelSequenceEditor/Source/LevelSequenceEditor/Public/LevelSequenceEditorModule.h
+++ b/Engine/Plugins/MovieScene/LevelSequenceEditor/Source/LevelSequenceEditor/Public/LevelSequenceEditorModule.h
@@ -7,7 +7,6 @@
 
 
 class UObject;
-struct FLevelSequenceActionExtender;
 
 
 /**
@@ -19,7 +18,4 @@ class ILevelSequenceEditorModule
 public:
 	DECLARE_EVENT_OneParam(ILevelSequenceEditorModule, FOnMasterSequenceCreated, UObject*);
 	virtual FOnMasterSequenceCreated& OnMasterSequenceCreated() = 0;
-
-	virtual void RegisterLevelSequenceActionExtender(TSharedRef<FLevelSequenceActionExtender> InExtender) = 0;
-	virtual void UnregisterLevelSequenceActionExtender(TSharedRef<FLevelSequenceActionExtender> InExtender) = 0;
 };
diff --git a/Engine/Plugins/Runtime/SoundMod/Source/SoundModImporter/Private/AssetTypeActions_SoundMod.cpp b/Engine/Plugins/Runtime/SoundMod/Source/SoundModImporter/Private/AssetTypeActions_SoundMod.cpp
--- a/Engine/Plugins/Runtime/SoundMod/Source/SoundModImporter/Private/AssetTypeActions_SoundMod.cpp
+++ b/Engine/Plugins/Runtime/SoundMod/Source/SoundModImporter/Private/AssetTypeActions_SoundMod.cpp
@@ -1,7 +1,7 @@
 // Copyright 1998-2019 Epic Games, Inc. All Rights Reserved.
 
 #include "AssetTypeActions_SoundMod.h"
-#include "Framework/MultiBox/MultiBoxBuilder.h"
+#include "EditorMenuSubsystem.h"
 #include "Editor.h"
 #include "Components/AudioComponent.h"
 #include "SoundMod.h"
@@ -13,11 +13,12 @@ UClass* FAssetTypeActions_SoundMod::GetSupportedClass() const
 	return USoundMod::StaticClass();
 }
 
-void FAssetTypeActions_SoundMod::GetActions(const TArray<UObject*>& InObjects, FMenuBuilder& MenuBuilder)
+void FAssetTypeActions_SoundMod::GetActions(const TArray<UObject*>& InObjects, FEditorMenuSection& Section)
 {
 	auto Sounds = GetTypedWeakObjectPtrs<USoundMod>(InObjects);
 
-	MenuBuilder.AddMenuEntry(
+	Section.AddMenuEntry(
+		"Sound_PlaySound",
 		LOCTEXT("Sound_PlaySound", "Play"),
 		LOCTEXT("Sound_PlaySoundTooltip", "Plays the selected sound."),
 		FSlateIcon(),
@@ -27,7 +28,8 @@ void FAssetTypeActions_SoundMod::GetActions(const TArray<UObject*>& InObjects, F
 		)
 		);
 
-	MenuBuilder.AddMenuEntry(
+	Section.AddMenuEntry(
+		"Sound_StopSound",
 		LOCTEXT("Sound_StopSound", "Stop"),
 		LOCTEXT("Sound_StopSoundTooltip", "Stops the selected sounds."),
 		FSlateIcon(),
diff --git a/Engine/Plugins/Runtime/SoundMod/Source/SoundModImporter/Private/AssetTypeActions_SoundMod.h b/Engine/Plugins/Runtime/SoundMod/Source/SoundModImporter/Private/AssetTypeActions_SoundMod.h
--- a/Engine/Plugins/Runtime/SoundMod/Source/SoundModImporter/Private/AssetTypeActions_SoundMod.h
+++ b/Engine/Plugins/Runtime/SoundMod/Source/SoundModImporter/Private/AssetTypeActions_SoundMod.h
@@ -5,7 +5,6 @@
 #include "CoreMinimal.h"
 #include "AssetTypeActions_Base.h"
 
-class FMenuBuilder;
 class USoundMod;
 
 class FAssetTypeActions_SoundMod: public FAssetTypeActions_Base
@@ -16,7 +15,7 @@ public:
 	virtual FColor GetTypeColor() const override{ return FColor(255, 175, 0); }
 	virtual UClass* GetSupportedClass() const override;
 	virtual bool HasActions(const TArray<UObject*>& InObjects) const override{ return true; }
-	virtual void GetActions(const TArray<UObject*>& InObjects, FMenuBuilder& MenuBuilder) override;
+	virtual void GetActions(const TArray<UObject*>& InObjects, struct FEditorMenuSection& Section) override;
 	virtual void AssetsActivated(const TArray<UObject*>& InObjects, EAssetTypeActivationMethod::Type ActivationType) override;
 	virtual uint32 GetCategories() override{ return EAssetTypeCategories::Sounds; }
 	virtual bool CanFilter() override { return false; }
diff --git a/Engine/Plugins/Runtime/SoundMod/Source/SoundModImporter/SoundModImporter.Build.cs b/Engine/Plugins/Runtime/SoundMod/Source/SoundModImporter/SoundModImporter.Build.cs
--- a/Engine/Plugins/Runtime/SoundMod/Source/SoundModImporter/SoundModImporter.Build.cs
+++ b/Engine/Plugins/Runtime/SoundMod/Source/SoundModImporter/SoundModImporter.Build.cs
@@ -14,7 +14,8 @@ public class SoundModImporter : ModuleRules
 				"SlateCore",
                 "Slate",
 				"SoundMod",
-				"UnrealEd"
+				"UnrealEd",
+				"EditorMenus",
 			});
 
 		PrivateIncludePathModuleNames.AddRange(
diff --git a/Engine/Plugins/Runtime/SoundUtilities/Source/SoundUtilitiesEditor/Private/SoundUtilitiesEditorModule.cpp b/Engine/Plugins/Runtime/SoundUtilities/Source/SoundUtilitiesEditor/Private/SoundUtilitiesEditorModule.cpp
--- a/Engine/Plugins/Runtime/SoundUtilities/Source/SoundUtilitiesEditor/Private/SoundUtilitiesEditorModule.cpp
+++ b/Engine/Plugins/Runtime/SoundUtilities/Source/SoundUtilitiesEditor/Private/SoundUtilitiesEditorModule.cpp
@@ -8,26 +8,29 @@
 #include "AssetToolsModule.h"
 #include "AssetTypeActions_SoundSimple.h"
 #include "AssetTypeActions_Base.h"
+#include "EditorMenuSubsystem.h"
 #include "AudioEditorModule.h"
-
+#include "SoundWaveAssetActionExtender.h"
 
 IMPLEMENT_MODULE(FSoundUtilitiesEditorModule, SoundUtilitiesEditor)
 
 void FSoundUtilitiesEditorModule::StartupModule()
 {
-	SoundWaveAssetActionExtender = TSharedPtr<ISoundWaveAssetActionExtensions>(new FSoundWaveAssetActionExtender());
-
 	IAssetTools& AssetTools = FModuleManager::LoadModuleChecked<FAssetToolsModule>("AssetTools").Get();
-
-	// Now that we've loaded this module, we need to register our effect preset actions
-	IAudioEditorModule* AudioEditorModule = &FModuleManager::LoadModuleChecked<IAudioEditorModule>("AudioEditor");
-	AudioEditorModule->AddSoundWaveActionExtender(SoundWaveAssetActionExtender);
-
+	
 	// Register asset actions
 	AssetTools.RegisterAssetTypeActions(MakeShareable(new FAssetTypeActions_SoundSimple));
+
+	UEditorMenuSubsystem::RegisterStartupCallback(FSimpleMulticastDelegate::FDelegate::CreateRaw(this, &FSoundUtilitiesEditorModule::RegisterMenus));
 }
 
 void FSoundUtilitiesEditorModule::ShutdownModule()
 {
+	UEditorMenuSubsystem::UnRegisterStartupCallback(this);
+	UEditorMenuSubsystem::UnregisterOwner("SoundUtilities");
 }
 
+void FSoundUtilitiesEditorModule::RegisterMenus()
+{
+	FSoundWaveAssetActionExtender::RegisterMenus();
+}
diff --git a/Engine/Plugins/Runtime/SoundUtilities/Source/SoundUtilitiesEditor/Private/SoundWaveAssetActionExtender.cpp b/Engine/Plugins/Runtime/SoundUtilities/Source/SoundUtilitiesEditor/Private/SoundWaveAssetActionExtender.cpp
--- a/Engine/Plugins/Runtime/SoundUtilities/Source/SoundUtilitiesEditor/Private/SoundWaveAssetActionExtender.cpp
+++ b/Engine/Plugins/Runtime/SoundUtilities/Source/SoundUtilitiesEditor/Private/SoundWaveAssetActionExtender.cpp
@@ -1,12 +1,12 @@
 // Copyright 1998-2019 Epic Games, Inc. All Rights Reserved.
 
 #include "SoundWaveAssetActionExtender.h"
-#include "Framework/Commands/UIAction.h"
-#include "Framework/Commands/UICommandList.h"
-#include "Framework/MultiBox/MultiBoxBuilder.h"
+#include "EditorMenuSubsystem.h"
 #include "AssetTypeActions_Base.h"
 #include "IContentBrowserSingleton.h"
 #include "ContentBrowserModule.h"
+#include "ContentBrowserMenuContexts.h"
+#include "ObjectEditorUtils.h"
 #include "EditorStyleSet.h"
 #include "Sound/SoundWave.h"
 #include "SoundSimple.h"
@@ -14,24 +14,40 @@
 
 #define LOCTEXT_NAMESPACE "AssetTypeActions"
 
-void FSoundWaveAssetActionExtender::GetExtendedActions(const TArray<TWeakObjectPtr<USoundWave>>& InSounds, FMenuBuilder& MenuBuilder)
+void FSoundWaveAssetActionExtender::RegisterMenus()
 {
-	const TAttribute<FText> Label = LOCTEXT("SoundWave_CreateSimpleSound", "Create Simple Sound");
-	const TAttribute<FText> ToolTip = LOCTEXT("SoundWave_CreateSimpleSoundTooltip", "Creates a simple sound asset using the selected sound waves.");
-	const FSlateIcon Icon = FSlateIcon(FEditorStyle::GetStyleSetName(), "ClassIcon.SoundSimple");
-	const FUIAction UIAction = FUIAction(FExecuteAction::CreateSP(this, &FSoundWaveAssetActionExtender::ExecuteCreateSimpleSound, InSounds), FCanExecuteAction());
+	if (!UEditorMenuSubsystem::IsRunningEditorUI())
+	{
+		return;
+	}
 
-	MenuBuilder.AddMenuEntry(Label, ToolTip, Icon, UIAction);
+	FEditorMenuOwnerScoped MenuOwner("SoundUtilities");
+	UEditorMenu* Menu = UEditorMenuSubsystem::Get()->ExtendMenu("ContentBrowser.AssetContextMenu.SoundWave");
+	FEditorMenuSection& Section = Menu->FindOrAddSection("GetAssetActions");
+	
+	Section.AddDynamicEntry("SoundWaveAsset", FNewEditorMenuSectionDelegate::CreateLambda([](FEditorMenuSection& InSection)
+	{
+		const TAttribute<FText> Label = LOCTEXT("SoundWave_CreateSimpleSound", "Create Simple Sound");
+		const TAttribute<FText> ToolTip = LOCTEXT("SoundWave_CreateSimpleSoundTooltip", "Creates a simple sound asset using the selected sound waves.");
+		const FSlateIcon Icon = FSlateIcon(FEditorStyle::GetStyleSetName(), "ClassIcon.SoundSimple");
+		const FEditorMenuExecuteAction UIAction = FEditorMenuExecuteAction::CreateStatic(&FSoundWaveAssetActionExtender::ExecuteCreateSimpleSound);
+
+		InSection.AddMenuEntry("SoundWave_CreateSimpleSound", Label, ToolTip, Icon, UIAction);
+	}));
 }
 
-void FSoundWaveAssetActionExtender::ExecuteCreateSimpleSound(TArray<TWeakObjectPtr<USoundWave>> SoundWaves)
+void FSoundWaveAssetActionExtender::ExecuteCreateSimpleSound(const FEditorMenuContext& MenuContext)
 {
-	const FString DefaultSuffix = TEXT("_SimpleSound");
-
-	if (SoundWaves.Num() > 0)
+	UContentBrowserAssetContextMenuContext* Context = MenuContext.Find<UContentBrowserAssetContextMenuContext>();
+	if (!Context || Context->SelectedObjects.Num() == 0)
 	{
-		USoundWave* SoundWave = SoundWaves[0].Get();
+		return;
+	}
 
+	const FString DefaultSuffix = TEXT("_SimpleSound");
+	
+	if (USoundWave* SoundWave = Cast<USoundWave>(Context->SelectedObjects[0]))
+	{
 		// Determine an appropriate name
 		FString Name;
 		FString PackagePath;
@@ -41,10 +57,13 @@ void FSoundWaveAssetActionExtender::ExecuteCreateSimpleSound(TArray<TWeakObjectP
 
 		// Create the factory used to generate the asset
 		USoundSimpleFactory* Factory = NewObject<USoundSimpleFactory>();
-
-		for (int32 i = 0; i < SoundWaves.Num(); ++i)
+		Factory->SoundWaves.Reset();
+		for (UObject* Object : Context->SelectedObjects)
 		{
-			Factory->SoundWaves.Add(SoundWaves[i].Get());
+			if (USoundWave* Wave = Cast<USoundWave>(Object))
+			{
+				Factory->SoundWaves.Add(Wave);
+			}
 		}
 
 		FContentBrowserModule& ContentBrowserModule = FModuleManager::LoadModuleChecked<FContentBrowserModule>("ContentBrowser");
diff --git a/Engine/Plugins/Runtime/SoundUtilities/Source/SoundUtilitiesEditor/Public/SoundUtilitiesEditorModule.h b/Engine/Plugins/Runtime/SoundUtilities/Source/SoundUtilitiesEditor/Public/SoundUtilitiesEditorModule.h
--- a/Engine/Plugins/Runtime/SoundUtilities/Source/SoundUtilitiesEditor/Public/SoundUtilitiesEditorModule.h
+++ b/Engine/Plugins/Runtime/SoundUtilities/Source/SoundUtilitiesEditor/Public/SoundUtilitiesEditorModule.h
@@ -4,9 +4,6 @@
 
 #include "Modules/ModuleInterface.h"
 #include "Modules/ModuleManager.h"
-#include "SoundWaveAssetActionExtender.h"
-
-class ISoundWaveAssetActionExtensions;
 
 class FSoundUtilitiesEditorModule : public IModuleInterface
 {
@@ -16,5 +13,6 @@ public:
 	virtual void StartupModule() override;
 	virtual void ShutdownModule() override;
 
-	TSharedPtr<ISoundWaveAssetActionExtensions> SoundWaveAssetActionExtender;
+private:
+	void RegisterMenus();
 };
diff --git a/Engine/Plugins/Runtime/SoundUtilities/Source/SoundUtilitiesEditor/Public/SoundWaveAssetActionExtender.h b/Engine/Plugins/Runtime/SoundUtilities/Source/SoundUtilitiesEditor/Public/SoundWaveAssetActionExtender.h
--- a/Engine/Plugins/Runtime/SoundUtilities/Source/SoundUtilitiesEditor/Public/SoundWaveAssetActionExtender.h
+++ b/Engine/Plugins/Runtime/SoundUtilities/Source/SoundUtilitiesEditor/Public/SoundWaveAssetActionExtender.h
@@ -2,17 +2,11 @@
 
 #pragma once
 
-#include "AudioEditorModule.h"
-
-
-class FSoundWaveAssetActionExtender :	public ISoundWaveAssetActionExtensions,
-												public TSharedFromThis<FSoundWaveAssetActionExtender>
+class FSoundWaveAssetActionExtender
 {
 public:
-	virtual ~FSoundWaveAssetActionExtender() {}
-
-	virtual void GetExtendedActions(const TArray<TWeakObjectPtr<USoundWave>>& InObjects, FMenuBuilder& MenuBuilder) override;
-
-	void ExecuteCreateSimpleSound(TArray<TWeakObjectPtr<USoundWave>> Objects);
+	static void RegisterMenus();
+	static void GetExtendedActions(const struct FEditorMenuContext& MenuContext);
+	static void ExecuteCreateSimpleSound(const struct FEditorMenuContext& MenuContext);
 };
 
diff --git a/Engine/Plugins/Runtime/SoundUtilities/Source/SoundUtilitiesEditor/SoundUtilitiesEditor.Build.cs b/Engine/Plugins/Runtime/SoundUtilities/Source/SoundUtilitiesEditor/SoundUtilitiesEditor.Build.cs
--- a/Engine/Plugins/Runtime/SoundUtilities/Source/SoundUtilitiesEditor/SoundUtilitiesEditor.Build.cs
+++ b/Engine/Plugins/Runtime/SoundUtilities/Source/SoundUtilitiesEditor/SoundUtilitiesEditor.Build.cs
@@ -19,6 +19,7 @@ namespace UnrealBuildTool.Rules
                     "Slate",
                     "SlateCore",
                     "ContentBrowser",
+					"EditorMenus",
 				}
 			);
 
diff --git a/Engine/Plugins/Runtime/TimeSynth/Source/TimeSynthEditor/Private/TimeSynthEditorModule.cpp b/Engine/Plugins/Runtime/TimeSynth/Source/TimeSynthEditor/Private/TimeSynthEditorModule.cpp
--- a/Engine/Plugins/Runtime/TimeSynth/Source/TimeSynthEditor/Private/TimeSynthEditorModule.cpp
+++ b/Engine/Plugins/Runtime/TimeSynth/Source/TimeSynthEditor/Private/TimeSynthEditorModule.cpp
@@ -11,7 +11,7 @@
 #include "TimeSynthClip.h"
 #include "TimeSynthSoundWaveAssetActionExtender.h"
 #include "TimeSynthVolumeGroup.h"
-
+#include "EditorMenuSubsystem.h"
 #include "AudioEditorModule.h"
 
 
@@ -19,18 +19,20 @@ IMPLEMENT_MODULE(FTimeSynthEditorModule, TimeSynthEditor)
 
 void FTimeSynthEditorModule::StartupModule()
 {
-	TimeSynthSoundWaveAssetActionExtender = TSharedPtr<ISoundWaveAssetActionExtensions>(new FTimeSynthSoundWaveAssetActionExtender());
-
 	IAssetTools& AssetTools = FModuleManager::LoadModuleChecked<FAssetToolsModule>("AssetTools").Get();
-
-	// Now that we've loaded this module, we need to register our effect preset actions
-	IAudioEditorModule* AudioEditorModule = &FModuleManager::LoadModuleChecked<IAudioEditorModule>("AudioEditor");
-	AudioEditorModule->AddSoundWaveActionExtender(TimeSynthSoundWaveAssetActionExtender);
-
 	AssetTools.RegisterAssetTypeActions(MakeShareable(new FAssetTypeActions_TimeSynthClip));
 	AssetTools.RegisterAssetTypeActions(MakeShareable(new FAssetTypeActions_TimeSynthVolumeGroup));
+
+	UEditorMenuSubsystem::RegisterStartupCallback(FSimpleMulticastDelegate::FDelegate::CreateRaw(this, &FTimeSynthEditorModule::RegisterMenus));
 }
 
 void FTimeSynthEditorModule::ShutdownModule()
 {
+	UEditorMenuSubsystem::UnRegisterStartupCallback(this);
+	UEditorMenuSubsystem::UnregisterOwner("TimeSynth");
+}
+
+void FTimeSynthEditorModule::RegisterMenus()
+{
+	FTimeSynthSoundWaveAssetActionExtender::RegisterMenus();
 }
diff --git a/Engine/Plugins/Runtime/TimeSynth/Source/TimeSynthEditor/Private/TimeSynthSoundWaveAssetActionExtender.cpp b/Engine/Plugins/Runtime/TimeSynth/Source/TimeSynthEditor/Private/TimeSynthSoundWaveAssetActionExtender.cpp
--- a/Engine/Plugins/Runtime/TimeSynth/Source/TimeSynthEditor/Private/TimeSynthSoundWaveAssetActionExtender.cpp
+++ b/Engine/Plugins/Runtime/TimeSynth/Source/TimeSynthEditor/Private/TimeSynthSoundWaveAssetActionExtender.cpp
@@ -2,12 +2,12 @@
 
 #include "TimeSynthSoundWaveAssetActionExtender.h"
 
-#include "Framework/Commands/UIAction.h"
-#include "Framework/Commands/UICommandList.h"
-#include "Framework/MultiBox/MultiBoxBuilder.h"
+#include "EditorMenuSubsystem.h"
 #include "AssetTypeActions_Base.h"
 #include "IContentBrowserSingleton.h"
 #include "ContentBrowserModule.h"
+#include "ContentBrowserMenuContexts.h"
+#include "ObjectEditorUtils.h"
 #include "EditorStyleSet.h"
 #include "Sound/SoundWave.h"
 #include "TimeSynthClip.h"
@@ -20,29 +20,56 @@ namespace
 	const FString DefaultSuffix = TEXT("_TSC");
 } // namespace <>
 
-void FTimeSynthSoundWaveAssetActionExtender::GetExtendedActions(const TArray<TWeakObjectPtr<USoundWave>>& InSounds, FMenuBuilder& MenuBuilder)
+void FTimeSynthSoundWaveAssetActionExtender::RegisterMenus()
 {
-	const TAttribute<FText> Label = LOCTEXT("SoundWave_CreateTimeSynthClip", "Create Time Synth Clip(s)");
-	const TAttribute<FText> ToolTip = LOCTEXT("SoundWave_CreateTimeSynthClipToolTip", "Creates time synth clip per sound wave selected");
-	const FSlateIcon Icon = FSlateIcon(FEditorStyle::GetStyleSetName(), "ClassIcon.SoundCue");
-	const FUIAction UIAction = FUIAction(FExecuteAction::CreateSP(this, &FTimeSynthSoundWaveAssetActionExtender::ExecuteCreateTimeSyncClip, InSounds), FCanExecuteAction());
+	if (!UEditorMenuSubsystem::IsRunningEditorUI())
+	{
+		return;
+	}
+
+	FEditorMenuOwnerScoped MenuOwner("TimeSynth");
+	UEditorMenu* Menu = UEditorMenuSubsystem::Get()->ExtendMenu("ContentBrowser.AssetContextMenu.SoundWave");
+	FEditorMenuSection& Section = Menu->FindOrAddSection("GetAssetActions");
+
+	Section.AddDynamicEntry("TimeSynthSoundWaveAsset", FNewEditorMenuSectionDelegate::CreateLambda([](FEditorMenuSection& InSection)
+	{
+		UContentBrowserAssetContextMenuContext* Context = InSection.FindContext<UContentBrowserAssetContextMenuContext>();
+		if (!Context || Context->SelectedObjects.Num() == 0)
+		{
+			return;
+		}
 
-	MenuBuilder.AddMenuEntry(Label, ToolTip, Icon, UIAction);
+		const TAttribute<FText> Label = LOCTEXT("SoundWave_CreateTimeSynthClip", "Create Time Synth Clip(s)");
+		const TAttribute<FText> ToolTip = LOCTEXT("SoundWave_CreateTimeSynthClipToolTip", "Creates time synth clip per sound wave selected");
+		const FSlateIcon Icon = FSlateIcon(FEditorStyle::GetStyleSetName(), "ClassIcon.SoundCue");
 
-	const TAttribute<FText> LabelSet = LOCTEXT("SoundWave_CreateTimeSynthClipSet", "Create Time Synth Clip Set");
-	const TAttribute<FText> ToolTipSet = LOCTEXT("SoundWave_CreateTimeSynthClipSetToolTip", "Creates time synth clip adding all selected sound waves to single clip as set.");
-	const FUIAction UIActionSet = FUIAction(FExecuteAction::CreateSP(this, &FTimeSynthSoundWaveAssetActionExtender::ExecuteCreateTimeSyncClipSet, InSounds), FCanExecuteAction());
+		const FEditorMenuExecuteAction UIExecuteAction = FEditorMenuExecuteAction::CreateStatic(&FTimeSynthSoundWaveAssetActionExtender::ExecuteCreateTimeSyncClip);
 
-	MenuBuilder.AddMenuEntry(LabelSet, ToolTipSet, Icon, UIActionSet);
+		InSection.AddMenuEntry("SoundWave_CreateTimeSynthClip", Label, ToolTip, Icon, UIExecuteAction);
+
+		const TAttribute<FText> LabelSet = LOCTEXT("SoundWave_CreateTimeSynthClipSet", "Create Time Synth Clip Set");
+		const TAttribute<FText> ToolTipSet = LOCTEXT("SoundWave_CreateTimeSynthClipSetToolTip", "Creates time synth clip adding all selected sound waves to single clip as set.");
+		const FEditorMenuExecuteAction UIExecuteActionSet = FEditorMenuExecuteAction::CreateStatic(&FTimeSynthSoundWaveAssetActionExtender::ExecuteCreateTimeSyncClipSet);
+
+		InSection.AddMenuEntry("SoundWave_CreateTimeSynthClipSet", LabelSet, ToolTipSet, Icon, UIExecuteActionSet);
+	}));
 }
 
-void FTimeSynthSoundWaveAssetActionExtender::ExecuteCreateTimeSyncClip(TArray<TWeakObjectPtr<USoundWave>> SoundWaves)
+void FTimeSynthSoundWaveAssetActionExtender::ExecuteCreateTimeSyncClip(const FEditorMenuContext& MenuContext)
 {
+	UContentBrowserAssetContextMenuContext* Context = MenuContext.Find<UContentBrowserAssetContextMenuContext>();
+	if (!Context || Context->SelectedObjects.Num() == 0)
+	{
+		return;
+	}
+
 	FAssetToolsModule& AssetToolsModule = FModuleManager::Get().LoadModuleChecked<FAssetToolsModule>("AssetTools");
 
 	UTimeSynthClipFactory* Factory = NewObject<UTimeSynthClipFactory>();
-	for (TWeakObjectPtr<USoundWave>& SoundWave : SoundWaves)
+	for (UObject* Object : Context->SelectedObjects)
 	{
+		USoundWave* SoundWave = Cast<USoundWave>(Object);
+
 		FString Name;
 		FString PackagePath;
 
@@ -55,22 +82,28 @@ void FTimeSynthSoundWaveAssetActionExtender::ExecuteCreateTimeSyncClip(TArray<TW
 	}
 }
 
-void FTimeSynthSoundWaveAssetActionExtender::ExecuteCreateTimeSyncClipSet(TArray<TWeakObjectPtr<USoundWave>> SoundWaves)
+void FTimeSynthSoundWaveAssetActionExtender::ExecuteCreateTimeSyncClipSet(const FEditorMenuContext& MenuContext)
 {
-	FString PackagePath;
-	FString Name;
+	UContentBrowserAssetContextMenuContext* Context = MenuContext.Find<UContentBrowserAssetContextMenuContext>();
+	if (!Context || Context->SelectedObjects.Num() == 0)
+	{
+		return;
+	}
 
-	if (SoundWaves.Num() == 0 || !SoundWaves[0].IsValid())
+	USoundWave* SoundWave = Cast<USoundWave>(Context->SelectedObjects[0]);
+	if (SoundWave == nullptr)
 	{
 		return;
 	}
 
-	USoundWave& SoundWave = *SoundWaves[0].Get();
+	FString PackagePath;
+	FString Name;
+
 	FAssetToolsModule& AssetToolsModule = FModuleManager::Get().LoadModuleChecked<FAssetToolsModule>("AssetTools");
-	AssetToolsModule.Get().CreateUniqueAssetName(SoundWave.GetOutermost()->GetName(), DefaultSuffix, PackagePath, Name);
+	AssetToolsModule.Get().CreateUniqueAssetName(SoundWave->GetOutermost()->GetName(), DefaultSuffix, PackagePath, Name);
 
 	UTimeSynthClipFactory* Factory = NewObject<UTimeSynthClipFactory>();
-	Factory->SoundWaves = SoundWaves;
+	Factory->SoundWaves = FObjectEditorUtils::GetTypedWeakObjectPtrs<USoundWave>(Context->SelectedObjects);
 
 	AssetToolsModule.Get().CreateAsset(Name, FPackageName::GetLongPackagePath(PackagePath), UTimeSynthClip::StaticClass(), Factory);
 }
diff --git a/Engine/Plugins/Runtime/TimeSynth/Source/TimeSynthEditor/Public/TimeSynthEditorModule.h b/Engine/Plugins/Runtime/TimeSynth/Source/TimeSynthEditor/Public/TimeSynthEditorModule.h
--- a/Engine/Plugins/Runtime/TimeSynth/Source/TimeSynthEditor/Public/TimeSynthEditorModule.h
+++ b/Engine/Plugins/Runtime/TimeSynth/Source/TimeSynthEditor/Public/TimeSynthEditorModule.h
@@ -15,5 +15,6 @@ public:
 	virtual void ShutdownModule() override;
 
 private:
-	TSharedPtr<ISoundWaveAssetActionExtensions> TimeSynthSoundWaveAssetActionExtender;
+
+	void RegisterMenus();
 };
diff --git a/Engine/Plugins/Runtime/TimeSynth/Source/TimeSynthEditor/Public/TimeSynthSoundWaveAssetActionExtender.h b/Engine/Plugins/Runtime/TimeSynth/Source/TimeSynthEditor/Public/TimeSynthSoundWaveAssetActionExtender.h
--- a/Engine/Plugins/Runtime/TimeSynth/Source/TimeSynthEditor/Public/TimeSynthSoundWaveAssetActionExtender.h
+++ b/Engine/Plugins/Runtime/TimeSynth/Source/TimeSynthEditor/Public/TimeSynthSoundWaveAssetActionExtender.h
@@ -2,18 +2,11 @@
 
 #pragma once
 
-#include "AudioEditorModule.h"
-
-
-class FTimeSynthSoundWaveAssetActionExtender :	public ISoundWaveAssetActionExtensions,
-												public TSharedFromThis<FTimeSynthSoundWaveAssetActionExtender>
+class FTimeSynthSoundWaveAssetActionExtender
 {
 public:
-	virtual ~FTimeSynthSoundWaveAssetActionExtender() {}
-
-	virtual void GetExtendedActions(const TArray<TWeakObjectPtr<USoundWave>>& InObjects, FMenuBuilder& MenuBuilder) override;
-
-	void ExecuteCreateTimeSyncClip(TArray<TWeakObjectPtr<USoundWave>> Objects);
-	void ExecuteCreateTimeSyncClipSet(TArray<TWeakObjectPtr<USoundWave>> SoundWaves);
+	static void RegisterMenus();
+	static void ExecuteCreateTimeSyncClip(const struct FEditorMenuContext& MenuContext);
+	static void ExecuteCreateTimeSyncClipSet(const struct FEditorMenuContext& MenuContext);
 };
 
diff --git a/Engine/Plugins/Runtime/TimeSynth/Source/TimeSynthEditor/TimeSynthEditor.Build.cs b/Engine/Plugins/Runtime/TimeSynth/Source/TimeSynthEditor/TimeSynthEditor.Build.cs
--- a/Engine/Plugins/Runtime/TimeSynth/Source/TimeSynthEditor/TimeSynthEditor.Build.cs
+++ b/Engine/Plugins/Runtime/TimeSynth/Source/TimeSynthEditor/TimeSynthEditor.Build.cs
@@ -18,6 +18,8 @@ namespace UnrealBuildTool.Rules
                     "EditorStyle",
                     "Slate",
                     "SlateCore",
+					"ContentBrowser",
+					"EditorMenus",
 				}
 			);
 
diff --git a/Engine/Plugins/VirtualProduction/Takes/Source/TakeRecorder/Private/TakeRecorderModule.cpp b/Engine/Plugins/VirtualProduction/Takes/Source/TakeRecorder/Private/TakeRecorderModule.cpp
--- a/Engine/Plugins/VirtualProduction/Takes/Source/TakeRecorder/Private/TakeRecorderModule.cpp
+++ b/Engine/Plugins/VirtualProduction/Takes/Source/TakeRecorder/Private/TakeRecorderModule.cpp
@@ -16,6 +16,8 @@
 #include "Framework/Docking/LayoutExtender.h"
 #include "Framework/MultiBox/MultiBoxExtender.h"
 #include "Framework/MultiBox/MultiBoxBuilder.h"
+#include "EditorMenuSubsystem.h"
+#include "ContentBrowserMenuContexts.h"
 #include "Features/IModularFeatures.h"
 #include "ITakeRecorderDropHandler.h"
 #include "ISettingsModule.h"
@@ -26,7 +28,6 @@
 #include "ISequencer.h"
 #include "LevelSequence.h"
 #include "LevelSequenceEditorModule.h"
-#include "LevelSequenceActionExtender.h"
 #include "SequencerSettings.h"
 #include "TakeMetaData.h"
 #include "FileHelpers.h"
@@ -164,14 +165,36 @@ namespace
 	}
 }
 
-struct FTakeRecorderLevelSequenceActionExtender : FLevelSequenceActionExtender
+void FTakeRecorderModule::RegisterMenus()
 {
-	virtual void GetActions(const TArray<UObject*>& InObjects, FMenuBuilder& MenuBuilder) override
+#if WITH_EDITOR
+	if (!UEditorMenuSubsystem::IsRunningEditorUI())
+	{
+		return;
+	}
+
+	FEditorMenuOwnerScoped MenuOwner("TakeRecorder");
+	UEditorMenuSubsystem* EditorMenus = UEditorMenuSubsystem::Get();
+	UEditorMenu* Menu = EditorMenus->ExtendMenu("ContentBrowser.AssetContextMenu.LevelSequence");
+	if (!Menu)
+	{
+		return;
+	}
+
+	FEditorMenuSection& Section = Menu->FindOrAddSection("GetAssetActions");
+	Section.AddDynamicEntry("TakeRecorderActions", FNewEditorMenuSectionDelegate::CreateLambda([](FEditorMenuSection& InSection)
 	{
-		ULevelSequence* LevelSequence = InObjects.Num() == 1 ? Cast<ULevelSequence>(InObjects[0]) : nullptr;
+		UContentBrowserAssetContextMenuContext* Context = InSection.FindContext<UContentBrowserAssetContextMenuContext>();
+		if (!Context)
+		{
+			return;
+		}
+
+		ULevelSequence* LevelSequence = Context->SelectedObjects.Num() == 1 ? Cast<ULevelSequence>(Context->SelectedObjects[0]) : nullptr;
 		if (LevelSequence)
 		{
-			MenuBuilder.AddMenuEntry(
+			InSection.AddMenuEntry(
+				"OpenInTakeRecorder_Label",
 				LOCTEXT("OpenInTakeRecorder_Label", "Open in Take Recorder"),
 				LOCTEXT("OpenInTakeRecorder_Tooltip", "Opens this level sequence asset in Take Recorder"),
 				FSlateIcon(FTakeRecorderStyle::StyleName, "TakeRecorder.TabIcon"),
@@ -217,7 +240,8 @@ struct FTakeRecorderLevelSequenceActionExtender : FLevelSequenceActionExtender
 			}
 
 			if (bHasMapTag) {
-				MenuBuilder.AddMenuEntry(
+				InSection.AddMenuEntry(
+					"TakeRecorderOpenMap_Label",
 					LOCTEXT("TakeRecorderOpenMap_Label", "Open Map"),
 					LOCTEXT("TakeRecorderOpenMap_Tooltip", "Opens the map used to create this Level Sequence Asset"),
 					FSlateIcon(FEditorStyle::GetStyleSetName(), "LevelEditor.Tabs.Levels" ),
@@ -242,8 +266,9 @@ struct FTakeRecorderLevelSequenceActionExtender : FLevelSequenceActionExtender
 				);
 			}
 		}
-	}
-};
+	}));
+#endif // WITH_EDITOR
+}
 
 FTakeRecorderModule::FTakeRecorderModule()
 	: SequencerSettings(nullptr)
@@ -260,10 +285,28 @@ void FTakeRecorderModule::StartupModule()
 	RegisterAssetTools();
 	RegisterSettings();
 	RegisterSerializedRecorder();
+
+#if WITH_EDITOR
+	if (GIsEditor)
+	{
+		if (UEditorMenuSubsystem::TryGet())
+		{
+			RegisterMenus();
+		}
+		else
+		{
+			FCoreDelegates::OnPostEngineInit.AddRaw(this, &FTakeRecorderModule::RegisterMenus);
+		}
+	}
+#endif
 }
 
 void FTakeRecorderModule::ShutdownModule()
 {
+#if WITH_EDITOR
+	FCoreDelegates::OnPostEngineInit.RemoveAll(this);
+#endif
+
 	FTakeRecorderCommands::Unregister();
 
 	UnregisterDetailCustomizations();
@@ -365,19 +408,6 @@ void FTakeRecorderModule::RegisterAssetTools()
 		TakePresetActions = MakeShared<FTakePresetActions>();
 		AssetTools.RegisterAssetTypeActions(TakePresetActions.ToSharedRef());
 	}
-
-#if WITH_EDITOR
-
-	if (GIsEditor)
-	{
-		ILevelSequenceEditorModule& LevelSequenceEditorModule = FModuleManager::LoadModuleChecked<ILevelSequenceEditorModule>("LevelSequenceEditor");
-		{
-			LevelSequenceAssetActionExtender = MakeShared<FTakeRecorderLevelSequenceActionExtender>();
-			LevelSequenceEditorModule.RegisterLevelSequenceActionExtender(LevelSequenceAssetActionExtender.ToSharedRef());
-		}
-	}
-
-#endif
 }
 
 void FTakeRecorderModule::UnregisterAssetTools()
@@ -387,12 +417,6 @@ void FTakeRecorderModule::UnregisterAssetTools()
 	{
 		AssetToolsModule->Get().UnregisterAssetTypeActions(TakePresetActions.ToSharedRef());
 	}
-
-	ILevelSequenceEditorModule* LevelSequenceEditorModule = FModuleManager::GetModulePtr<ILevelSequenceEditorModule>("LevelSequenceEditor");
-	if (LevelSequenceEditorModule)
-	{
-		LevelSequenceEditorModule->UnregisterLevelSequenceActionExtender(LevelSequenceAssetActionExtender.ToSharedRef());
-	}
 }
 
 void FTakeRecorderModule::RegisterSettings()
diff --git a/Engine/Plugins/VirtualProduction/Takes/Source/TakeRecorder/Private/TakeRecorderModule.h b/Engine/Plugins/VirtualProduction/Takes/Source/TakeRecorder/Private/TakeRecorderModule.h
--- a/Engine/Plugins/VirtualProduction/Takes/Source/TakeRecorder/Private/TakeRecorderModule.h
+++ b/Engine/Plugins/VirtualProduction/Takes/Source/TakeRecorder/Private/TakeRecorderModule.h
@@ -7,7 +7,6 @@
 #include "UObject/GCObject.h"
 #include "ITakeRecorderModule.h"
 
-struct FLevelSequenceActionExtender;
 
 class FExtender;
 class FTakePresetActions;
@@ -56,6 +55,8 @@ private:
 
 	void UnregisterSerializedRecorder();
 
+	void RegisterMenus();
+
 private:
 
 	DECLARE_MULTICAST_DELEGATE_TwoParams(FOnExtendSourcesMenuEvent, TSharedRef<FExtender>, UTakeRecorderSources*);
@@ -70,7 +71,6 @@ private:
 	FName ProjectSettingsName;
 
 	TSharedPtr<FTakePresetActions> TakePresetActions;
-	TSharedPtr<FLevelSequenceActionExtender> LevelSequenceAssetActionExtender;
 	TSharedPtr<FSerializedRecorder> SerializedRecorder;
 
 	USequencerSettings* SequencerSettings;
diff --git a/Engine/Plugins/VirtualProduction/Takes/Source/TakeRecorder/TakeRecorder.Build.cs b/Engine/Plugins/VirtualProduction/Takes/Source/TakeRecorder/TakeRecorder.Build.cs
--- a/Engine/Plugins/VirtualProduction/Takes/Source/TakeRecorder/TakeRecorder.Build.cs
+++ b/Engine/Plugins/VirtualProduction/Takes/Source/TakeRecorder/TakeRecorder.Build.cs
@@ -32,6 +32,7 @@ public class TakeRecorder : ModuleRules
 				"UnrealEd",
 				"WorkspaceMenuStructure",
 				"Analytics",
+				"EditorMenus",
 			}
 		);
 
diff --git a/Engine/Source/Developer/AssetTools/Private/AssetTools.cpp b/Engine/Source/Developer/AssetTools/Private/AssetTools.cpp
--- a/Engine/Source/Developer/AssetTools/Private/AssetTools.cpp
+++ b/Engine/Source/Developer/AssetTools/Private/AssetTools.cpp
@@ -26,6 +26,7 @@
 #include "AssetToolsLog.h"
 #include "AssetToolsModule.h"
 #include "Framework/MultiBox/MultiBoxBuilder.h"
+#include "EditorMenuSubsystem.h"
 #include "IClassTypeActions.h"
 #include "AssetTypeActions/AssetTypeActions_Blueprint.h"
 #include "AssetTypeActions/AssetTypeActions_Curve.h"
diff --git a/Engine/Source/Developer/AssetTools/Private/AssetTypeActions/AssetTypeActions_AnimBlueprint.cpp b/Engine/Source/Developer/AssetTools/Private/AssetTypeActions/AssetTypeActions_AnimBlueprint.cpp
--- a/Engine/Source/Developer/AssetTools/Private/AssetTypeActions/AssetTypeActions_AnimBlueprint.cpp
+++ b/Engine/Source/Developer/AssetTools/Private/AssetTypeActions/AssetTypeActions_AnimBlueprint.cpp
@@ -1,7 +1,7 @@
 // Copyright 1998-2019 Epic Games, Inc. All Rights Reserved.
 
 #include "AssetTypeActions/AssetTypeActions_AnimBlueprint.h"
-#include "Framework/MultiBox/MultiBoxBuilder.h"
+#include "EditorMenuSubsystem.h"
 #include "Widgets/Layout/SBorder.h"
 #include "Misc/MessageDialog.h"
 #include "Widgets/Images/SImage.h"
@@ -21,13 +21,14 @@
 
 #define LOCTEXT_NAMESPACE "AssetTypeActions"
 
-void FAssetTypeActions_AnimBlueprint::GetActions( const TArray<UObject*>& InObjects, FMenuBuilder& MenuBuilder )
+void FAssetTypeActions_AnimBlueprint::GetActions(const TArray<UObject*>& InObjects, FEditorMenuSection& Section)
 {
-	FAssetTypeActions_Blueprint::GetActions(InObjects, MenuBuilder);
+	FAssetTypeActions_Blueprint::GetActions(InObjects, Section);
 
 	auto AnimBlueprints = GetTypedWeakObjectPtrs<UAnimBlueprint>(InObjects);
 
-	MenuBuilder.AddMenuEntry(
+	Section.AddMenuEntry(
+		"AnimBlueprint_FindSkeleton",
 		LOCTEXT("AnimBlueprint_FindSkeleton", "Find Skeleton"),
 		LOCTEXT("AnimBlueprint_FindSkeletonTooltip", "Finds the skeleton used by the selected Anim Blueprints in the content browser."),
 		FSlateIcon(FEditorStyle::GetStyleSetName(), "Persona.AssetActions.FindSkeleton"),
@@ -37,17 +38,18 @@ void FAssetTypeActions_AnimBlueprint::GetActions( const TArray<UObject*>& InObje
 			)
 		);
 
-	MenuBuilder.AddSubMenu( 
+	Section.AddEntry(FEditorMenuEntry::InitSubMenu(
+		NAME_None,
+		"RetargetBlueprintSubmenu",
 		LOCTEXT("RetargetBlueprintSubmenu", "Retarget Anim Blueprints"),
 		LOCTEXT("RetargetBlueprintSubmenu_ToolTip", "Opens the retarget blueprints menu"),
-		FNewMenuDelegate::CreateSP( this, &FAssetTypeActions_AnimBlueprint::FillRetargetMenu, InObjects ),
+		FNewEditorMenuDelegate::CreateSP( this, &FAssetTypeActions_AnimBlueprint::FillRetargetMenu, InObjects),
 		false,
 		FSlateIcon(FEditorStyle::GetStyleSetName(), "Persona.AssetActions.RetargetSkeleton")
-		);
-
+		));
 }
 
-void FAssetTypeActions_AnimBlueprint::FillRetargetMenu( FMenuBuilder& MenuBuilder, const TArray<UObject*> InObjects )
+void FAssetTypeActions_AnimBlueprint::FillRetargetMenu(UEditorMenu* MenuBuilder, const TArray<UObject*> InObjects)
 {
 	bool bAllSkeletonsNull = true;
 
@@ -63,9 +65,11 @@ void FAssetTypeActions_AnimBlueprint::FillRetargetMenu( FMenuBuilder& MenuBuilde
 		}
 	}
 
+	FEditorMenuSection& Section = MenuBuilder->AddSection("Section");
 	if(bAllSkeletonsNull)
 	{
-		MenuBuilder.AddMenuEntry(
+		Section.AddMenuEntry(
+			"AnimBlueprint_RetargetSkeletonInPlace",
 			LOCTEXT("AnimBlueprint_RetargetSkeletonInPlace", "Retarget skeleton on existing Anim Blueprints"),
 			LOCTEXT("AnimBlueprint_RetargetSkeletonInPlaceTooltip", "Retargets the selected Anim Blueprints to a new skeleton (and optionally all referenced animations too)"),
 			FSlateIcon(FEditorStyle::GetStyleSetName(), "Persona.AssetActions.RetargetSkeleton"),
@@ -76,7 +80,8 @@ void FAssetTypeActions_AnimBlueprint::FillRetargetMenu( FMenuBuilder& MenuBuilde
 			);
 	}
 
-	MenuBuilder.AddMenuEntry(
+	Section.AddMenuEntry(
+		"AnimBlueprint_DuplicateAndRetargetSkeleton",
 		LOCTEXT("AnimBlueprint_DuplicateAndRetargetSkeleton", "Duplicate Anim Blueprints and Retarget"),
 		LOCTEXT("AnimBlueprint_DuplicateAndRetargetSkeletonTooltip", "Duplicates and then retargets the selected Anim Blueprints to a new skeleton (and optionally all referenced animations too)"),
 		FSlateIcon(FEditorStyle::GetStyleSetName(), "Persona.AssetActions.DuplicateAndRetargetSkeleton"),
diff --git a/Engine/Source/Developer/AssetTools/Private/AssetTypeActions/AssetTypeActions_AnimBlueprint.h b/Engine/Source/Developer/AssetTools/Private/AssetTypeActions/AssetTypeActions_AnimBlueprint.h
--- a/Engine/Source/Developer/AssetTools/Private/AssetTypeActions/AssetTypeActions_AnimBlueprint.h
+++ b/Engine/Source/Developer/AssetTools/Private/AssetTypeActions/AssetTypeActions_AnimBlueprint.h
@@ -8,7 +8,6 @@
 #include "EditorAnimUtils.h"
 #include "Animation/AnimBlueprint.h"
 
-class FMenuBuilder;
 class UFactory;
 
 class FAssetTypeActions_AnimBlueprint : public FAssetTypeActions_Blueprint
@@ -18,7 +17,7 @@ public:
 	virtual FText GetName() const override { return NSLOCTEXT("AssetTypeActions", "AssetTypeActions_AnimBlueprint", "Animation Blueprint"); }
 	virtual FColor GetTypeColor() const override { return FColor(200,116,0); }
 	virtual UClass* GetSupportedClass() const override { return UAnimBlueprint::StaticClass(); }
-	virtual void GetActions( const TArray<UObject*>& InObjects, FMenuBuilder& MenuBuilder ) override;
+	virtual void GetActions(const TArray<UObject*>& InObjects, struct FEditorMenuSection& Section) override;
 	virtual void OpenAssetEditor( const TArray<UObject*>& InObjects, TSharedPtr<class IToolkitHost> EditWithinLevelEditor = TSharedPtr<IToolkitHost>() ) override;
 	virtual uint32 GetCategories() override { return EAssetTypeCategories::Animation; }
 	virtual void PerformAssetDiff(UObject* Asset1, UObject* Asset2, const struct FRevisionInfo& OldRevision, const struct FRevisionInfo& NewRevision) const override;
@@ -31,7 +30,7 @@ public:
 private:
 
 	/** Handler to fill the retarget submenu */
-	void FillRetargetMenu( FMenuBuilder& MenuBuilder, const TArray<UObject*> InObjects );
+	void FillRetargetMenu(class UEditorMenu* Menu, const TArray<UObject*> InObjects);
 
 	/** Handler for when FindSkeleton is selected */
 	void ExecuteFindSkeleton(TArray<TWeakObjectPtr<UAnimBlueprint>> Objects);
diff --git a/Engine/Source/Developer/AssetTools/Private/AssetTypeActions/AssetTypeActions_AnimCurveCompressionSettings.cpp b/Engine/Source/Developer/AssetTools/Private/AssetTypeActions/AssetTypeActions_AnimCurveCompressionSettings.cpp
--- a/Engine/Source/Developer/AssetTools/Private/AssetTypeActions/AssetTypeActions_AnimCurveCompressionSettings.cpp
+++ b/Engine/Source/Developer/AssetTools/Private/AssetTypeActions/AssetTypeActions_AnimCurveCompressionSettings.cpp
@@ -4,7 +4,7 @@
 #include "Animation/AnimSequence.h"
 #include "Dialogs/Dialogs.h"
 #include "EditorStyleSet.h"
-#include "Framework/MultiBox/MultiBoxBuilder.h"
+#include "EditorMenuSubsystem.h"
 #include "Misc/ScopedSlowTask.h"
 #include "UObject/UObjectIterator.h"
 
@@ -41,7 +41,7 @@ void FAssetTypeActions_AnimCurveCompressionSettings::AddToolbarExtension(FToolBa
 	Builder.EndSection();
 }
 
-void FAssetTypeActions_AnimCurveCompressionSettings::GetActions(const TArray<UObject*>& InObjects, FMenuBuilder& MenuBuilder)
+void FAssetTypeActions_AnimCurveCompressionSettings::GetActions(const TArray<UObject*>& InObjects, FEditorMenuSection& Section)
 {
 	auto SettingAssets = GetTypedWeakObjectPtrs<UAnimCurveCompressionSettings>(InObjects);
 
@@ -50,7 +50,8 @@ void FAssetTypeActions_AnimCurveCompressionSettings::GetActions(const TArray<UOb
 		return;
 	}
 
-	MenuBuilder.AddMenuEntry(
+	Section.AddMenuEntry(
+		"AnimCurveCompressionSettings_Compress",
 		LOCTEXT("AnimCurveCompressionSettings_Compress", "Compress"),
 		LOCTEXT("AnimCurveCompressionSettings_CompressTooltip", "All animation sequences that use these settings will be compressed."),
 		FSlateIcon(FEditorStyle::GetStyleSetName(), "Persona.ApplyCompression.Small"),
diff --git a/Engine/Source/Developer/AssetTools/Private/AssetTypeActions/AssetTypeActions_AnimCurveCompressionSettings.h b/Engine/Source/Developer/AssetTools/Private/AssetTypeActions/AssetTypeActions_AnimCurveCompressionSettings.h
--- a/Engine/Source/Developer/AssetTools/Private/AssetTypeActions/AssetTypeActions_AnimCurveCompressionSettings.h
+++ b/Engine/Source/Developer/AssetTools/Private/AssetTypeActions/AssetTypeActions_AnimCurveCompressionSettings.h
@@ -19,7 +19,7 @@ public:
 	virtual void OpenAssetEditor(const TArray<UObject*>& InObjects, TSharedPtr<class IToolkitHost> EditWithinLevelEditor = TSharedPtr<IToolkitHost>()) override;
 
 	virtual bool HasActions(const TArray<UObject*>& InObjects) const override { return true; }
-	virtual void GetActions(const TArray<UObject*>& InObjects, FMenuBuilder& MenuBuilder) override;
+	virtual void GetActions(const TArray<UObject*>& InObjects, struct FEditorMenuSection& Section) override;
 
 private:
 	void AddToolbarExtension(FToolBarBuilder& Builder, TWeakObjectPtr<UAnimCurveCompressionSettings> CurveSettings);
diff --git a/Engine/Source/Developer/AssetTools/Private/AssetTypeActions/AssetTypeActions_AnimMontage.cpp b/Engine/Source/Developer/AssetTools/Private/AssetTypeActions/AssetTypeActions_AnimMontage.cpp
--- a/Engine/Source/Developer/AssetTools/Private/AssetTypeActions/AssetTypeActions_AnimMontage.cpp
+++ b/Engine/Source/Developer/AssetTools/Private/AssetTypeActions/AssetTypeActions_AnimMontage.cpp
@@ -3,14 +3,14 @@
 #include "AssetTypeActions/AssetTypeActions_AnimMontage.h"
 #include "Factories/AnimMontageFactory.h"
 #include "AnimationEditorUtils.h"
-#include "Framework/MultiBox/MultiBoxBuilder.h"
+#include "EditorMenuSubsystem.h"
 #include "AssetTools.h"
 #include "EditorStyleSet.h"
 
 #define LOCTEXT_NAMESPACE "AssetTypeActions"
 
 
-void FAssetTypeActions_AnimMontage::GetActions( const TArray<UObject*>& InObjects, FMenuBuilder& MenuBuilder )
+void FAssetTypeActions_AnimMontage::GetActions(const TArray<UObject*>& InObjects, FEditorMenuSection& Section)
 {
 	auto Montages = GetTypedWeakObjectPtrs<UAnimMontage>(InObjects);
 
@@ -29,7 +29,8 @@ void FAssetTypeActions_AnimMontage::GetActions( const TArray<UObject*>& InObject
 	if (!bContainsChildMontage)
 	{
 		// create mew child anim montage
-		MenuBuilder.AddMenuEntry(
+		Section.AddMenuEntry(
+			"AnimMontage_CreateChildMontage",
 			LOCTEXT("AnimMontage_CreateChildMontage", "Create Child Montage"),
 			LOCTEXT("AnimMontage_CreateChildMontageTooltip", "Create Child Animation Montage and remap to another animation assets."),
 			FSlateIcon(FEditorStyle::GetStyleSetName(), "ClassIcon.AnimMontage"),
@@ -37,7 +38,7 @@ void FAssetTypeActions_AnimMontage::GetActions( const TArray<UObject*>& InObject
 		);
 	}
 
-	FAssetTypeActions_AnimationAsset::GetActions(InObjects, MenuBuilder);
+	FAssetTypeActions_AnimationAsset::GetActions(InObjects, Section);
 }
 
 void FAssetTypeActions_AnimMontage::CreateChildAnimMontage(TArray<TWeakObjectPtr<UAnimMontage>> AnimMontages) 
diff --git a/Engine/Source/Developer/AssetTools/Private/AssetTypeActions/AssetTypeActions_AnimMontage.h b/Engine/Source/Developer/AssetTools/Private/AssetTypeActions/AssetTypeActions_AnimMontage.h
--- a/Engine/Source/Developer/AssetTools/Private/AssetTypeActions/AssetTypeActions_AnimMontage.h
+++ b/Engine/Source/Developer/AssetTools/Private/AssetTypeActions/AssetTypeActions_AnimMontage.h
@@ -6,8 +6,6 @@
 #include "AssetTypeActions/AssetTypeActions_AnimationAsset.h"
 #include "Animation/AnimMontage.h"
 
-class FMenuBuilder;
-
 class FAssetTypeActions_AnimMontage : public FAssetTypeActions_AnimationAsset
 {
 public:
@@ -16,7 +14,7 @@ public:
 	virtual FColor GetTypeColor() const override { return FColor(100,100,255); }
 	virtual UClass* GetSupportedClass() const override { return UAnimMontage::StaticClass(); }
 	virtual bool CanFilter() override { return true; }
-	virtual void GetActions(const TArray<UObject*>& InObjects, FMenuBuilder& MenuBuilder) override;
+	virtual void GetActions(const TArray<UObject*>& InObjects, struct FEditorMenuSection& Section) override;
 
 private:
    	/* 
diff --git a/Engine/Source/Developer/AssetTools/Private/AssetTypeActions/AssetTypeActions_AnimSequence.cpp b/Engine/Source/Developer/AssetTools/Private/AssetTypeActions/AssetTypeActions_AnimSequence.cpp
--- a/Engine/Source/Developer/AssetTools/Private/AssetTypeActions/AssetTypeActions_AnimSequence.cpp
+++ b/Engine/Source/Developer/AssetTools/Private/AssetTypeActions/AssetTypeActions_AnimSequence.cpp
@@ -2,7 +2,7 @@
 
 #include "AssetTypeActions/AssetTypeActions_AnimSequence.h"
 #include "Animation/AnimSequence.h"
-#include "Framework/MultiBox/MultiBoxBuilder.h"
+#include "EditorMenuSubsystem.h"
 #include "EditorStyleSet.h"
 #include "EditorReimportHandler.h"
 #include "Animation/AnimMontage.h"
@@ -27,34 +27,38 @@ UClass* FAssetTypeActions_AnimSequence::GetSupportedClass() const
 	return UAnimSequence::StaticClass(); 
 }
 
-void FAssetTypeActions_AnimSequence::GetActions( const TArray<UObject*>& InObjects, FMenuBuilder& MenuBuilder )
+void FAssetTypeActions_AnimSequence::GetActions(const TArray<UObject*>& InObjects, FEditorMenuSection& Section)
 {
 	auto Sequences = GetTypedWeakObjectPtrs<UAnimSequence>(InObjects);
 
 	// create menu
-	MenuBuilder.AddSubMenu(
+	Section.AddEntry(FEditorMenuEntry::InitSubMenu(
+		NAME_None,
+		"CreateAnimSubmenu",
 		LOCTEXT("CreateAnimSubmenu", "Create"),
 		LOCTEXT("CreateAnimSubmenu_ToolTip", "Create assets from this anim sequence"),
 		FNewMenuDelegate::CreateSP(this, &FAssetTypeActions_AnimSequence::FillCreateMenu, Sequences),
 		false,
 		FSlateIcon(FEditorStyle::GetStyleSetName(), "Persona.AssetActions.CreateAnimAsset")
-		);
+		));
 
-	MenuBuilder.AddMenuEntry(
+	Section.AddMenuEntry(
+		"AnimSequence_ReimportWithNewSource",
 		LOCTEXT("AnimSequence_ReimportWithNewSource", "Reimport with New Source"),
 		LOCTEXT("AnimSequence_ReimportWithNewSourceTooltip", "Reimport the selected sequence(s) from a new source file."),
 		FSlateIcon(FEditorStyle::GetStyleSetName(), "Persona.AssetActions.ReimportAnim"),
 		FUIAction(FExecuteAction::CreateSP(this, &FAssetTypeActions_AnimSequence::ExecuteReimportWithNewSource, Sequences))
 		);
 
-	MenuBuilder.AddMenuEntry(
+	Section.AddMenuEntry(
+		"AnimSequence_AddAnimationModifier",
 		LOCTEXT("AnimSequence_AddAnimationModifier", "Add Animation Modifier(s)"),
 		LOCTEXT("AnimSequence_AddAnimationModifierTooltip", "Apply new animation modifier(s)."),
 		FSlateIcon(FEditorStyle::GetStyleSetName(), "ClassIcon.AnimationModifier"),
 		FUIAction(FExecuteAction::CreateSP(this, &FAssetTypeActions_AnimSequence::ExecuteAddNewAnimationModifier, Sequences))
 	);
 
-	FAssetTypeActions_AnimationAsset::GetActions(InObjects, MenuBuilder);
+	FAssetTypeActions_AnimationAsset::GetActions(InObjects, Section);
 }
 
 void FAssetTypeActions_AnimSequence::FillCreateMenu(FMenuBuilder& MenuBuilder, const TArray<TWeakObjectPtr<UAnimSequence>> Sequences) const
diff --git a/Engine/Source/Developer/AssetTools/Private/AssetTypeActions/AssetTypeActions_AnimSequence.h b/Engine/Source/Developer/AssetTools/Private/AssetTypeActions/AssetTypeActions_AnimSequence.h
--- a/Engine/Source/Developer/AssetTools/Private/AssetTypeActions/AssetTypeActions_AnimSequence.h
+++ b/Engine/Source/Developer/AssetTools/Private/AssetTypeActions/AssetTypeActions_AnimSequence.h
@@ -7,7 +7,6 @@
 #include "Animation/AnimationAsset.h"
 #include "AssetTypeActions/AssetTypeActions_AnimationAsset.h"
 
-class FMenuBuilder;
 class UAnimSequence;
 class UFactory;
 
@@ -20,7 +19,7 @@ public:
 	// IAssetTypeActions Implementation
 	virtual FText GetName() const override { return NSLOCTEXT("AssetTypeActions", "AssetTypeActions_AnimSequence", "Animation Sequence"); }
 	virtual UClass* GetSupportedClass() const override;
-	virtual void GetActions( const TArray<UObject*>& InObjects, FMenuBuilder& MenuBuilder ) override;
+	virtual void GetActions(const TArray<UObject*>& InObjects, struct FEditorMenuSection& Section) override;
 	virtual bool CanFilter() override { return true; }
 	virtual bool IsImportedAsset() const override { return true; }
 	virtual void GetResolvedSourceFilePaths(const TArray<UObject*>& TypeAssets, TArray<FString>& OutSourceFilePaths) const override;
diff --git a/Engine/Source/Developer/AssetTools/Private/AssetTypeActions/AssetTypeActions_AnimStreamable.h b/Engine/Source/Developer/AssetTools/Private/AssetTypeActions/AssetTypeActions_AnimStreamable.h
--- a/Engine/Source/Developer/AssetTools/Private/AssetTypeActions/AssetTypeActions_AnimStreamable.h
+++ b/Engine/Source/Developer/AssetTools/Private/AssetTypeActions/AssetTypeActions_AnimStreamable.h
@@ -16,10 +16,10 @@ public:
 	virtual FColor GetTypeColor() const override { return FColor(181,230,29); }
 	virtual UClass* GetSupportedClass() const override { return UAnimStreamable::StaticClass(); }
 	virtual bool CanFilter() override { return true; }
-	virtual void GetActions(const TArray<UObject*>& InObjects, FMenuBuilder& MenuBuilder) override;
+	virtual void GetActions(const TArray<UObject*>& InObjects, struct FEditorMenuSection& Section) override;
 };
 
-void FAssetTypeActions_AnimStreamable::GetActions(const TArray<UObject*>& InObjects, FMenuBuilder& MenuBuilder)
+void FAssetTypeActions_AnimStreamable::GetActions(const TArray<UObject*>& InObjects, FEditorMenuSection& Section)
 {
 	/*auto Sequences = GetTypedWeakObjectPtrs<UAnimStreamable>(InObjects);
 
diff --git a/Engine/Source/Developer/AssetTools/Private/AssetTypeActions/AssetTypeActions_AnimationAsset.cpp b/Engine/Source/Developer/AssetTools/Private/AssetTypeActions/AssetTypeActions_AnimationAsset.cpp
--- a/Engine/Source/Developer/AssetTools/Private/AssetTypeActions/AssetTypeActions_AnimationAsset.cpp
+++ b/Engine/Source/Developer/AssetTools/Private/AssetTypeActions/AssetTypeActions_AnimationAsset.cpp
@@ -12,18 +12,19 @@
 #include "SSkeletonWidget.h"
 #include "IAnimationEditorModule.h"
 #include "Preferences/PersonaOptions.h"
-#include "Framework/MultiBox/MultiBoxBuilder.h"
+#include "EditorMenuSubsystem.h"
 #include "Framework/Application/SlateApplication.h"
 #include "Algo/Transform.h"
 #include "IPersonaToolkit.h"
 
 #define LOCTEXT_NAMESPACE "AssetTypeActions"
 
-void FAssetTypeActions_AnimationAsset::GetActions( const TArray<UObject*>& InObjects, FMenuBuilder& MenuBuilder )
+void FAssetTypeActions_AnimationAsset::GetActions(const TArray<UObject*>& InObjects, FEditorMenuSection& Section)
 {
 	auto AnimAssets = GetTypedWeakObjectPtrs<UAnimationAsset>(InObjects);
 
-	MenuBuilder.AddMenuEntry(
+	Section.AddMenuEntry(
+		"AnimSequenceBase_OpenInNewWindow",
 		LOCTEXT("AnimSequenceBase_OpenInNewWindow", "Open In New Window"),
 		LOCTEXT("AnimSequenceBase_OpenInNewWindowTooltip", "Will always open asset in a new window, and not re-use existing window. (Shift+Double-Click)"),
 		FSlateIcon(FEditorStyle::GetStyleSetName(), "ContentBrowser.AssetActions.OpenInExternalEditor"),
@@ -33,7 +34,8 @@ void FAssetTypeActions_AnimationAsset::GetActions( const TArray<UObject*>& InObj
 		)
 	);
 
-	MenuBuilder.AddMenuEntry(
+	Section.AddMenuEntry(
+		"AnimSequenceBase_FindSkeleton",
 		LOCTEXT("AnimSequenceBase_FindSkeleton", "Find Skeleton"),
 		LOCTEXT("AnimSequenceBase_FindSkeletonTooltip", "Finds the skeleton for the selected assets in the content browser."),
 		FSlateIcon(FEditorStyle::GetStyleSetName(), "Persona.AssetActions.FindSkeleton"),
@@ -43,13 +45,15 @@ void FAssetTypeActions_AnimationAsset::GetActions( const TArray<UObject*>& InObj
 			)
 		);
 
-	MenuBuilder.AddSubMenu( 
+	Section.AddEntry(FEditorMenuEntry::InitSubMenu(
+		NAME_None,
+		"RetargetAnimSubmenu", 
 		LOCTEXT("RetargetAnimSubmenu", "Retarget Anim Assets"),
 		LOCTEXT("RetargetAnimSubmenu_ToolTip", "Opens the retarget anim assets menu"),
 		FNewMenuDelegate::CreateSP( this, &FAssetTypeActions_AnimationAsset::FillRetargetMenu, InObjects ),
 		false,
 		FSlateIcon(FEditorStyle::GetStyleSetName(), "Persona.AssetActions.RetargetSkeleton")
-		);
+		));
 }
 
 void FAssetTypeActions_AnimationAsset::FillRetargetMenu( FMenuBuilder& MenuBuilder, const TArray<UObject*> InObjects )
diff --git a/Engine/Source/Developer/AssetTools/Private/AssetTypeActions/AssetTypeActions_AnimationAsset.h b/Engine/Source/Developer/AssetTools/Private/AssetTypeActions/AssetTypeActions_AnimationAsset.h
--- a/Engine/Source/Developer/AssetTools/Private/AssetTypeActions/AssetTypeActions_AnimationAsset.h
+++ b/Engine/Source/Developer/AssetTools/Private/AssetTypeActions/AssetTypeActions_AnimationAsset.h
@@ -8,8 +8,6 @@
 #include "Animation/AnimationAsset.h"
 #include "EditorAnimUtils.h"
 
-class FMenuBuilder;
-
 class FAssetTypeActions_AnimationAsset : public FAssetTypeActions_Base
 {
 public:
@@ -18,7 +16,7 @@ public:
 	virtual FColor GetTypeColor() const override { return FColor(80,123,72); }
 	virtual UClass* GetSupportedClass() const override { return UAnimationAsset::StaticClass(); }
 	virtual bool HasActions ( const TArray<UObject*>& InObjects ) const override { return true; }
-	virtual void GetActions( const TArray<UObject*>& InObjects, FMenuBuilder& MenuBuilder ) override;
+	virtual void GetActions(const TArray<UObject*>& InObjects, struct FEditorMenuSection& Section) override;
 	virtual void OpenAssetEditor( const TArray<UObject*>& InObjects, TSharedPtr<class IToolkitHost> EditWithinLevelEditor = TSharedPtr<IToolkitHost>() ) override;
 	virtual bool CanFilter() override { return false; }
 	virtual uint32 GetCategories() override { return EAssetTypeCategories::Animation; }
diff --git a/Engine/Source/Developer/AssetTools/Private/AssetTypeActions/AssetTypeActions_Blueprint.cpp b/Engine/Source/Developer/AssetTools/Private/AssetTypeActions/AssetTypeActions_Blueprint.cpp
--- a/Engine/Source/Developer/AssetTools/Private/AssetTypeActions/AssetTypeActions_Blueprint.cpp
+++ b/Engine/Source/Developer/AssetTools/Private/AssetTypeActions/AssetTypeActions_Blueprint.cpp
@@ -1,7 +1,7 @@
 // Copyright 1998-2019 Epic Games, Inc. All Rights Reserved.
 
 #include "AssetTypeActions/AssetTypeActions_Blueprint.h"
-#include "Framework/MultiBox/MultiBoxBuilder.h"
+#include "EditorMenuSubsystem.h"
 #include "GameFramework/Actor.h"
 #include "Misc/PackageName.h"
 #include "Misc/MessageDialog.h"
@@ -19,7 +19,7 @@
 
 #define LOCTEXT_NAMESPACE "AssetTypeActions"
 
-void FAssetTypeActions_Blueprint::GetActions( const TArray<UObject*>& InObjects, FMenuBuilder& MenuBuilder )
+void FAssetTypeActions_Blueprint::GetActions(const TArray<UObject*>& InObjects, FEditorMenuSection& Section)
 {
 	TArray<TWeakObjectPtr<UBlueprint>> Blueprints = GetTypedWeakObjectPtrs<UBlueprint>(InObjects);
 	
@@ -38,16 +38,17 @@ void FAssetTypeActions_Blueprint::GetActions( const TArray<UObject*>& InObjects,
 
 		if (bCanEditSharedDefaults)
 		{
-			MenuBuilder.AddMenuEntry(
-			LOCTEXT("Blueprint_EditDefaults", "Edit Shared Defaults"),
-			LOCTEXT("Blueprint_EditDefaultsTooltip", "Edit the shared default properties of the selected blueprints."),
-			FSlateIcon(FEditorStyle::GetStyleSetName(), "Kismet.Tabs.BlueprintDefaults"),
-			FUIAction(
-				FExecuteAction::CreateSP( this, &FAssetTypeActions_Blueprint::ExecuteEditDefaults, Blueprints ),
-				FCanExecuteAction()
-				)
-			);
-		}	
+			Section.AddMenuEntry(
+				"Blueprint_EditDefaults",
+				LOCTEXT("Blueprint_EditDefaults", "Edit Shared Defaults"),
+				LOCTEXT("Blueprint_EditDefaultsTooltip", "Edit the shared default properties of the selected blueprints."),
+				FSlateIcon(FEditorStyle::GetStyleSetName(), "Kismet.Tabs.BlueprintDefaults"),
+				FUIAction(
+					FExecuteAction::CreateSP( this, &FAssetTypeActions_Blueprint::ExecuteEditDefaults, Blueprints ),
+					FCanExecuteAction()
+					)
+				);
+		}
 	}
 
 	if ( Blueprints.Num() == 1 && CanCreateNewDerivedBlueprint() )
@@ -56,7 +57,8 @@ void FAssetTypeActions_Blueprint::GetActions( const TArray<UObject*>& InObjects,
 		DynamicTooltipGetter.BindSP(this, &FAssetTypeActions_Blueprint::GetNewDerivedBlueprintTooltip, Blueprints[0]);
 		TAttribute<FText> DynamicTooltipAttribute = TAttribute<FText>::Create(DynamicTooltipGetter);
 
-		MenuBuilder.AddMenuEntry(
+		Section.AddMenuEntry(
+			"Blueprint_NewDerivedBlueprint",
 			LOCTEXT("Blueprint_NewDerivedBlueprint", "Create Child Blueprint Class"),
 			DynamicTooltipAttribute,
 			FSlateIcon(FEditorStyle::GetStyleSetName(), "LevelEditor.CreateClassBlueprint"),
diff --git a/Engine/Source/Developer/AssetTools/Private/AssetTypeActions/AssetTypeActions_Class.cpp b/Engine/Source/Developer/AssetTools/Private/AssetTypeActions/AssetTypeActions_Class.cpp
--- a/Engine/Source/Developer/AssetTools/Private/AssetTypeActions/AssetTypeActions_Class.cpp
+++ b/Engine/Source/Developer/AssetTools/Private/AssetTypeActions/AssetTypeActions_Class.cpp
@@ -1,7 +1,7 @@
 // Copyright 1998-2019 Epic Games, Inc. All Rights Reserved.
 
 #include "AssetTypeActions/AssetTypeActions_Class.h"
-#include "Framework/MultiBox/MultiBoxBuilder.h"
+#include "EditorMenuSubsystem.h"
 #include "AssetData.h"
 #include "HAL/FileManager.h"
 #include "EditorStyleSet.h"
@@ -17,7 +17,7 @@ bool FAssetTypeActions_Class::HasActions(const TArray<UObject*>& InObjects) cons
 	return true;
 }
 
-void FAssetTypeActions_Class::GetActions(const TArray<UObject*>& InObjects, FMenuBuilder& MenuBuilder)
+void FAssetTypeActions_Class::GetActions(const TArray<UObject*>& InObjects, FEditorMenuSection& Section)
 {
 	UClass *const BaseClass = (InObjects.Num() == 1) ? Cast<UClass>(InObjects[0]) : nullptr;
 
@@ -103,7 +103,8 @@ void FAssetTypeActions_Class::GetActions(const TArray<UObject*>& InObjects, FMen
 		NewDerivedBlueprintClassToolTip = LOCTEXT("Class_NewDerivedBlueprintClassTooltip_InvalidNumberOfBases", "Can only create a Blueprint class when there is a single base class selected.");
 	}
 
-	MenuBuilder.AddMenuEntry(
+	Section.AddMenuEntry(
+		"NewDerivedCppClass",
 		NewDerivedCppClassLabel,
 		NewDerivedCppClassToolTip,
 		FSlateIcon(FEditorStyle::GetStyleSetName(), "MainFrame.AddCodeToProject"),
@@ -113,7 +114,8 @@ void FAssetTypeActions_Class::GetActions(const TArray<UObject*>& InObjects, FMen
 			)
 		);
 
-	MenuBuilder.AddMenuEntry(
+	Section.AddMenuEntry(
+		"NewDerivedBlueprintClass",
 		NewDerivedBlueprintClassLabel,
 		NewDerivedBlueprintClassToolTip,
 		FSlateIcon(FEditorStyle::GetStyleSetName(), "LevelEditor.CreateClassBlueprint"),
diff --git a/Engine/Source/Developer/AssetTools/Private/AssetTypeActions/AssetTypeActions_Class.h b/Engine/Source/Developer/AssetTools/Private/AssetTypeActions/AssetTypeActions_Class.h
--- a/Engine/Source/Developer/AssetTools/Private/AssetTypeActions/AssetTypeActions_Class.h
+++ b/Engine/Source/Developer/AssetTools/Private/AssetTypeActions/AssetTypeActions_Class.h
@@ -8,7 +8,6 @@
 #include "AssetTypeActions/AssetTypeActions_ClassTypeBase.h"
 
 struct FAssetData;
-class FMenuBuilder;
 class IClassTypeActions;
 
 class FAssetTypeActions_Class : public FAssetTypeActions_ClassTypeBase
@@ -21,7 +20,7 @@ public:
 	virtual uint32 GetCategories() override { return EAssetTypeCategories::Basic; }
 	
 	virtual bool HasActions(const TArray<UObject*>& InObjects) const override;
-	virtual void GetActions(const TArray<UObject*>& InObjects, FMenuBuilder& MenuBuilder) override;
+	virtual void GetActions(const TArray<UObject*>& InObjects, struct FEditorMenuSection& Section) override;
 
 	virtual class UThumbnailInfo* GetThumbnailInfo(UObject* Asset) const override;
 
diff --git a/Engine/Source/Developer/AssetTools/Private/AssetTypeActions/AssetTypeActions_CurveLinearColorAtlas.cpp b/Engine/Source/Developer/AssetTools/Private/AssetTypeActions/AssetTypeActions_CurveLinearColorAtlas.cpp
--- a/Engine/Source/Developer/AssetTools/Private/AssetTypeActions/AssetTypeActions_CurveLinearColorAtlas.cpp
+++ b/Engine/Source/Developer/AssetTools/Private/AssetTypeActions/AssetTypeActions_CurveLinearColorAtlas.cpp
@@ -1,7 +1,7 @@
 // Copyright 1998-2019 Epic Games, Inc. All Rights Reserved.
 
 #include "AssetTypeActions/AssetTypeActions_CurveLinearColorAtlas.h"
-#include "Framework/MultiBox/MultiBoxBuilder.h"
+#include "EditorMenuSubsystem.h"
 #include "Engine/Texture2D.h"
 #include "EditorStyleSet.h"
 #include "Engine/TextureCube.h"
@@ -11,9 +11,9 @@
 
 #define LOCTEXT_NAMESPACE "AssetTypeActions"
 
-void FAssetTypeActions_CurveLinearColorAtlas::GetActions( const TArray<UObject*>& InObjects, FMenuBuilder& MenuBuilder )
+void FAssetTypeActions_CurveLinearColorAtlas::GetActions(const TArray<UObject*>& InObjects, FEditorMenuSection& Section)
 {
-	FAssetTypeActions_Texture::GetActions(InObjects, MenuBuilder);
+	FAssetTypeActions_Texture::GetActions(InObjects, Section);
 }
 
 
diff --git a/Engine/Source/Developer/AssetTools/Private/AssetTypeActions/AssetTypeActions_CurveLinearColorAtlas.h b/Engine/Source/Developer/AssetTools/Private/AssetTypeActions/AssetTypeActions_CurveLinearColorAtlas.h
--- a/Engine/Source/Developer/AssetTools/Private/AssetTypeActions/AssetTypeActions_CurveLinearColorAtlas.h
+++ b/Engine/Source/Developer/AssetTools/Private/AssetTypeActions/AssetTypeActions_CurveLinearColorAtlas.h
@@ -6,7 +6,6 @@
 #include "AssetTypeActions/AssetTypeActions_Texture.h"
 #include "Curves/CurveLinearColorAtlas.h"
 
-class FMenuBuilder;
 
 class FAssetTypeActions_CurveLinearColorAtlas : public FAssetTypeActions_Texture
 {
@@ -16,7 +15,7 @@ public:
 	virtual FColor GetTypeColor() const override { return FColor(128,64,64); }
 	virtual UClass* GetSupportedClass() const override { return UCurveLinearColorAtlas::StaticClass(); }
 	virtual bool HasActions ( const TArray<UObject*>& InObjects ) const override { return true; }
-	virtual void GetActions( const TArray<UObject*>& InObjects, FMenuBuilder& MenuBuilder ) override;
+	virtual void GetActions(const TArray<UObject*>& InObjects, struct FEditorMenuSection& Section) override;
 	virtual bool IsImportedAsset() const override { return false; }
 	virtual bool CanFilter() override { return true; }
 	virtual uint32 GetCategories() override { return EAssetTypeCategories::Misc; }
diff --git a/Engine/Source/Developer/AssetTools/Private/AssetTypeActions/AssetTypeActions_CurveTable.cpp b/Engine/Source/Developer/AssetTools/Private/AssetTypeActions/AssetTypeActions_CurveTable.cpp
--- a/Engine/Source/Developer/AssetTools/Private/AssetTypeActions/AssetTypeActions_CurveTable.cpp
+++ b/Engine/Source/Developer/AssetTools/Private/AssetTypeActions/AssetTypeActions_CurveTable.cpp
@@ -1,7 +1,7 @@
 // Copyright 1998-2019 Epic Games, Inc. All Rights Reserved.
 
 #include "AssetTypeActions/AssetTypeActions_CurveTable.h"
-#include "Framework/MultiBox/MultiBoxBuilder.h"
+#include "EditorMenuSubsystem.h"
 #include "Misc/FileHelper.h"
 #include "EditorFramework/AssetImportData.h"
 #include "Framework/Application/SlateApplication.h"
@@ -11,7 +11,7 @@
 
 #define LOCTEXT_NAMESPACE "AssetTypeActions"
 
-void FAssetTypeActions_CurveTable::GetActions( const TArray<UObject*>& InObjects, FMenuBuilder& MenuBuilder )
+void FAssetTypeActions_CurveTable::GetActions(const TArray<UObject*>& InObjects, FEditorMenuSection& Section)
 {
 	auto Tables = GetTypedWeakObjectPtrs<UObject>(InObjects);
 
@@ -25,7 +25,8 @@ void FAssetTypeActions_CurveTable::GetActions( const TArray<UObject*>& InObjects
 		}
 	}
 
-	MenuBuilder.AddMenuEntry(
+	Section.AddMenuEntry(
+		"CurveTable_ExportAsCSV",
 		LOCTEXT("CurveTable_ExportAsCSV", "Export as CSV"),
 		LOCTEXT("CurveTable_ExportAsCSVTooltip", "Export the curve table as a file containing CSV data."),
 		FSlateIcon(),
@@ -35,7 +36,8 @@ void FAssetTypeActions_CurveTable::GetActions( const TArray<UObject*>& InObjects
 			)
 		);
 
-	MenuBuilder.AddMenuEntry(
+	Section.AddMenuEntry(
+		"CurveTable_ExportAsJSON",
 		LOCTEXT("CurveTable_ExportAsJSON", "Export as JSON"),
 		LOCTEXT("CurveTable_ExportAsJSONTooltip", "Export the curve table as a file containing JSON data."),
 		FSlateIcon(),
@@ -50,7 +52,8 @@ void FAssetTypeActions_CurveTable::GetActions( const TArray<UObject*>& InObjects
 	PotentialFileExtensions.Add(TEXT(".xlsm"));
 	PotentialFileExtensions.Add(TEXT(".csv"));
 	PotentialFileExtensions.Add(TEXT(".json"));
-	MenuBuilder.AddMenuEntry(
+	Section.AddMenuEntry(
+		"CurveTable_OpenSourceData",
 		LOCTEXT("CurveTable_OpenSourceData", "Open Source Data"),
 		LOCTEXT("CurveTable_OpenSourceDataTooltip", "Opens the curve table's source data file in an external editor. It will search using the following extensions: .xls/.xlsm/.csv/.json"),
 		FSlateIcon(),
diff --git a/Engine/Source/Developer/AssetTools/Private/AssetTypeActions/AssetTypeActions_CurveTable.h b/Engine/Source/Developer/AssetTools/Private/AssetTypeActions/AssetTypeActions_CurveTable.h
--- a/Engine/Source/Developer/AssetTools/Private/AssetTypeActions/AssetTypeActions_CurveTable.h
+++ b/Engine/Source/Developer/AssetTools/Private/AssetTypeActions/AssetTypeActions_CurveTable.h
@@ -7,7 +7,6 @@
 #include "Engine/CurveTable.h"
 #include "AssetTypeActions_CSVAssetBase.h"
 
-class FMenuBuilder;
 
 class FAssetTypeActions_CurveTable : public FAssetTypeActions_CSVAssetBase
 {
@@ -15,7 +14,7 @@ public:
 	// IAssetTypeActions Implementation
 	virtual FText GetName() const override { return NSLOCTEXT("AssetTypeActions", "AssetTypeActions_CurveTable", "Curve Table"); }
 	virtual UClass* GetSupportedClass() const override { return UCurveTable::StaticClass(); }
-	virtual void GetActions( const TArray<UObject*>& InObjects, FMenuBuilder& MenuBuilder ) override;
+	virtual void GetActions(const TArray<UObject*>& InObjects, struct FEditorMenuSection& Section) override;
 	virtual void OpenAssetEditor( const TArray<UObject*>& InObjects, TSharedPtr<class IToolkitHost> EditWithinLevelEditor = TSharedPtr<IToolkitHost>() ) override;
 	virtual void GetResolvedSourceFilePaths(const TArray<UObject*>& TypeAssets, TArray<FString>& OutSourceFilePaths) const override;
 	virtual void PerformAssetDiff(UObject* Asset1, UObject* Asset2, const struct FRevisionInfo& OldRevision, const struct FRevisionInfo& NewRevision) const override;
diff --git a/Engine/Source/Developer/AssetTools/Private/AssetTypeActions/AssetTypeActions_DataTable.cpp b/Engine/Source/Developer/AssetTools/Private/AssetTypeActions/AssetTypeActions_DataTable.cpp
--- a/Engine/Source/Developer/AssetTools/Private/AssetTypeActions/AssetTypeActions_DataTable.cpp
+++ b/Engine/Source/Developer/AssetTools/Private/AssetTypeActions/AssetTypeActions_DataTable.cpp
@@ -1,7 +1,7 @@
 // Copyright 1998-2019 Epic Games, Inc. All Rights Reserved.
 
 #include "AssetTypeActions/AssetTypeActions_DataTable.h"
-#include "Framework/MultiBox/MultiBoxBuilder.h"
+#include "EditorMenuSubsystem.h"
 #include "Misc/FileHelper.h"
 #include "EditorFramework/AssetImportData.h"
 #include "Dialogs/Dialogs.h"
@@ -13,7 +13,7 @@
 
 #define LOCTEXT_NAMESPACE "AssetTypeActions"
 
-void FAssetTypeActions_DataTable::GetActions( const TArray<UObject*>& InObjects, FMenuBuilder& MenuBuilder )
+void FAssetTypeActions_DataTable::GetActions(const TArray<UObject*>& InObjects, FEditorMenuSection& Section)
 {
 	auto Tables = GetTypedWeakObjectPtrs<UObject>(InObjects);
 	
@@ -27,7 +27,8 @@ void FAssetTypeActions_DataTable::GetActions( const TArray<UObject*>& InObjects,
 		}
 	}
 
-	MenuBuilder.AddMenuEntry(
+	Section.AddMenuEntry(
+		"DataTable_ExportAsCSV",
 		LOCTEXT("DataTable_ExportAsCSV", "Export as CSV"),
 		LOCTEXT("DataTable_ExportAsCSVTooltip", "Export the data table as a file containing CSV data."),
 		FSlateIcon(),
@@ -37,7 +38,8 @@ void FAssetTypeActions_DataTable::GetActions( const TArray<UObject*>& InObjects,
 			)
 		);
 
-	MenuBuilder.AddMenuEntry(
+	Section.AddMenuEntry(
+		"DataTable_ExportAsJSON",
 		LOCTEXT("DataTable_ExportAsJSON", "Export as JSON"),
 		LOCTEXT("DataTable_ExportAsJSONTooltip", "Export the data table as a file containing JSON data."),
 		FSlateIcon(),
@@ -52,7 +54,8 @@ void FAssetTypeActions_DataTable::GetActions( const TArray<UObject*>& InObjects,
 	PotentialFileExtensions.Add(TEXT(".xlsm"));
 	PotentialFileExtensions.Add(TEXT(".csv"));
 	PotentialFileExtensions.Add(TEXT(".json"));
-	MenuBuilder.AddMenuEntry(
+	Section.AddMenuEntry(
+		"DataTable_OpenSourceData",
 		LOCTEXT("DataTable_OpenSourceData", "Open Source Data"),
 		LOCTEXT("DataTable_OpenSourceDataTooltip", "Opens the data table's source data file in an external editor. It will search using the following extensions: .xls/.xlsm/.csv/.json"),
 		FSlateIcon(),
diff --git a/Engine/Source/Developer/AssetTools/Private/AssetTypeActions/AssetTypeActions_DataTable.h b/Engine/Source/Developer/AssetTools/Private/AssetTypeActions/AssetTypeActions_DataTable.h
--- a/Engine/Source/Developer/AssetTools/Private/AssetTypeActions/AssetTypeActions_DataTable.h
+++ b/Engine/Source/Developer/AssetTools/Private/AssetTypeActions/AssetTypeActions_DataTable.h
@@ -7,7 +7,6 @@
 #include "AssetTypeActions_CSVAssetBase.h"
 #include "Engine/DataTable.h"
 
-class FMenuBuilder;
 
 class FAssetTypeActions_DataTable : public FAssetTypeActions_CSVAssetBase
 {
@@ -15,7 +14,7 @@ public:
 	// IAssetTypeActions Implementation
 	virtual FText GetName() const override { return NSLOCTEXT("AssetTypeActions", "AssetTypeActions_DataTable", "Data Table"); }
 	virtual UClass* GetSupportedClass() const override { return UDataTable::StaticClass(); }
-	virtual void GetActions( const TArray<UObject*>& InObjects, FMenuBuilder& MenuBuilder ) override;
+	virtual void GetActions(const TArray<UObject*>& InObjects, struct FEditorMenuSection& Section) override;
 	virtual void OpenAssetEditor( const TArray<UObject*>& InObjects, TSharedPtr<class IToolkitHost> EditWithinLevelEditor = TSharedPtr<IToolkitHost>() ) override;
 	virtual void GetResolvedSourceFilePaths(const TArray<UObject*>& TypeAssets, TArray<FString>& OutSourceFilePaths) const override;
 	virtual void PerformAssetDiff(UObject* Asset1, UObject* Asset2, const struct FRevisionInfo& OldRevision, const struct FRevisionInfo& NewRevision) const override;
diff --git a/Engine/Source/Developer/AssetTools/Private/AssetTypeActions/AssetTypeActions_Font.cpp b/Engine/Source/Developer/AssetTools/Private/AssetTypeActions/AssetTypeActions_Font.cpp
--- a/Engine/Source/Developer/AssetTools/Private/AssetTypeActions/AssetTypeActions_Font.cpp
+++ b/Engine/Source/Developer/AssetTools/Private/AssetTypeActions/AssetTypeActions_Font.cpp
@@ -1,18 +1,19 @@
 // Copyright 1998-2019 Epic Games, Inc. All Rights Reserved.
 
 #include "AssetTypeActions/AssetTypeActions_Font.h"
-#include "Framework/MultiBox/MultiBoxBuilder.h"
+#include "EditorMenuSubsystem.h"
 #include "EditorStyleSet.h"
 #include "EditorReimportHandler.h"
 #include "FontEditorModule.h"
 
 #define LOCTEXT_NAMESPACE "AssetTypeActions"
 
-void FAssetTypeActions_Font::GetActions(const TArray<UObject*>& InObjects, FMenuBuilder& MenuBuilder)
+void FAssetTypeActions_Font::GetActions(const TArray<UObject*>& InObjects, FEditorMenuSection& Section)
 {
 	auto Fonts = GetTypedWeakObjectPtrs<UFont>(InObjects);
 
-	MenuBuilder.AddMenuEntry(
+	Section.AddMenuEntry(
+		"ReimportFont",
 		LOCTEXT("ReimportFontLabel", "Reimport"),
 		LOCTEXT("ReimportFontTooltip", "Reimport the selected font(s)."),
 		FSlateIcon(FEditorStyle::GetStyleSetName(), "ContentBrowser.AssetActions.ReimportAsset"),
diff --git a/Engine/Source/Developer/AssetTools/Private/AssetTypeActions/AssetTypeActions_Font.h b/Engine/Source/Developer/AssetTools/Private/AssetTypeActions/AssetTypeActions_Font.h
--- a/Engine/Source/Developer/AssetTools/Private/AssetTypeActions/AssetTypeActions_Font.h
+++ b/Engine/Source/Developer/AssetTools/Private/AssetTypeActions/AssetTypeActions_Font.h
@@ -7,7 +7,6 @@
 #include "AssetTypeActions_Base.h"
 #include "Engine/Font.h"
 
-class FMenuBuilder;
 
 class FAssetTypeActions_Font : public FAssetTypeActions_Base
 {
@@ -17,7 +16,7 @@ public:
 	virtual FColor GetTypeColor() const override { return FColor(128,128,64); }
 	virtual UClass* GetSupportedClass() const override { return UFont::StaticClass(); }
 	virtual bool HasActions(const TArray<UObject*>& InObjects) const override { return true; }
-	virtual void GetActions(const TArray<UObject*>& InObjects, FMenuBuilder& MenuBuilder) override;
+	virtual void GetActions(const TArray<UObject*>& InObjects, struct FEditorMenuSection& Section) override;
 	virtual void OpenAssetEditor( const TArray<UObject*>& InObjects, TSharedPtr<class IToolkitHost> EditWithinLevelEditor = TSharedPtr<IToolkitHost>() ) override;
 	virtual uint32 GetCategories() override { return EAssetTypeCategories::UI; }
 
diff --git a/Engine/Source/Developer/AssetTools/Private/AssetTypeActions/AssetTypeActions_FontFace.cpp b/Engine/Source/Developer/AssetTools/Private/AssetTypeActions/AssetTypeActions_FontFace.cpp
--- a/Engine/Source/Developer/AssetTools/Private/AssetTypeActions/AssetTypeActions_FontFace.cpp
+++ b/Engine/Source/Developer/AssetTools/Private/AssetTypeActions/AssetTypeActions_FontFace.cpp
@@ -2,17 +2,18 @@
 
 #include "AssetTypeActions/AssetTypeActions_FontFace.h"
 #include "FontEditorModule.h"
-#include "Framework/MultiBox/MultiBoxBuilder.h"
+#include "EditorMenuSubsystem.h"
 #include "EditorStyleSet.h"
 #include "EditorReimportHandler.h"
 
 #define LOCTEXT_NAMESPACE "AssetTypeActions"
 
-void FAssetTypeActions_FontFace::GetActions(const TArray<UObject*>& InObjects, FMenuBuilder& MenuBuilder)
+void FAssetTypeActions_FontFace::GetActions(const TArray<UObject*>& InObjects, FEditorMenuSection& Section)
 {
 	auto FontFaces = GetTypedWeakObjectPtrs<UFontFace>(InObjects);
 
-	MenuBuilder.AddMenuEntry(
+	Section.AddMenuEntry(
+		"ReimportFontFaceLabel",
 		LOCTEXT("ReimportFontFaceLabel", "Reimport"),
 		LOCTEXT("ReimportFontFaceTooltip", "Reimport the selected font(s)."),
 		FSlateIcon(FEditorStyle::GetStyleSetName(), "ContentBrowser.AssetActions.ReimportAsset"),
diff --git a/Engine/Source/Developer/AssetTools/Private/AssetTypeActions/AssetTypeActions_FontFace.h b/Engine/Source/Developer/AssetTools/Private/AssetTypeActions/AssetTypeActions_FontFace.h
--- a/Engine/Source/Developer/AssetTools/Private/AssetTypeActions/AssetTypeActions_FontFace.h
+++ b/Engine/Source/Developer/AssetTools/Private/AssetTypeActions/AssetTypeActions_FontFace.h
@@ -7,7 +7,6 @@
 #include "Toolkits/IToolkitHost.h"
 #include "AssetTypeActions_Base.h"
 
-class FMenuBuilder;
 
 class FAssetTypeActions_FontFace : public FAssetTypeActions_Base
 {
@@ -17,7 +16,7 @@ public:
 	virtual FColor GetTypeColor() const override { return FColor(184,184,112); }
 	virtual UClass* GetSupportedClass() const override { return UFontFace::StaticClass(); }
 	virtual bool HasActions(const TArray<UObject*>& InObjects) const override { return true; }
-	virtual void GetActions(const TArray<UObject*>& InObjects, FMenuBuilder& MenuBuilder) override;
+	virtual void GetActions(const TArray<UObject*>& InObjects, struct FEditorMenuSection& Section) override;
 	virtual void OpenAssetEditor( const TArray<UObject*>& InObjects, TSharedPtr<class IToolkitHost> EditWithinLevelEditor = TSharedPtr<IToolkitHost>() ) override;
 	virtual uint32 GetCategories() override { return EAssetTypeCategories::UI; }
 	virtual bool CanLocalize() const override { return false; }
diff --git a/Engine/Source/Developer/AssetTools/Private/AssetTypeActions/AssetTypeActions_ForceFeedbackEffect.cpp b/Engine/Source/Developer/AssetTools/Private/AssetTypeActions/AssetTypeActions_ForceFeedbackEffect.cpp
--- a/Engine/Source/Developer/AssetTools/Private/AssetTypeActions/AssetTypeActions_ForceFeedbackEffect.cpp
+++ b/Engine/Source/Developer/AssetTools/Private/AssetTypeActions/AssetTypeActions_ForceFeedbackEffect.cpp
@@ -1,7 +1,7 @@
 // Copyright 1998-2019 Epic Games, Inc. All Rights Reserved.
 
 #include "AssetTypeActions/AssetTypeActions_ForceFeedbackEffect.h"
-#include "Framework/MultiBox/MultiBoxBuilder.h"
+#include "EditorMenuSubsystem.h"
 #include "EditorStyleSet.h"
 #include "GenericPlatform/IInputInterface.h"
 #include "AssetData.h"
@@ -17,11 +17,12 @@ UClass* FAssetTypeActions_ForceFeedbackEffect::GetSupportedClass() const
 	return UForceFeedbackEffect::StaticClass();
 }
 
-void FAssetTypeActions_ForceFeedbackEffect::GetActions( const TArray<UObject*>& InObjects, FMenuBuilder& MenuBuilder )
+void FAssetTypeActions_ForceFeedbackEffect::GetActions(const TArray<UObject*>& InObjects, FEditorMenuSection& Section)
 {
 	TArray<TWeakObjectPtr<UForceFeedbackEffect>> Effects = GetTypedWeakObjectPtrs<UForceFeedbackEffect>(InObjects);
 
-	MenuBuilder.AddMenuEntry(
+	Section.AddMenuEntry(
+		"ForceFeedbackEffect_PlayEffect",
 		LOCTEXT("ForceFeedbackEffect_PlayEffect", "Play"),
 		LOCTEXT("ForceFeedbackEffect_PlayEffectTooltip", "Plays the selected force feedback effect."),
 		FSlateIcon(FEditorStyle::GetStyleSetName(), "MediaAsset.AssetActions.Play.Small"),
@@ -31,7 +32,8 @@ void FAssetTypeActions_ForceFeedbackEffect::GetActions( const TArray<UObject*>&
 			)
 		);
 
-	MenuBuilder.AddMenuEntry(
+	Section.AddMenuEntry(
+		"ForceFeedbackEffect_StopEffect",
 		LOCTEXT("ForceFeedbackEffect_StopEffect", "Stop"),
 		LOCTEXT("ForceFeedbackEffect_StopEffectTooltip", "Stops the selected force feedback effect."),
 		FSlateIcon(FEditorStyle::GetStyleSetName(), "MediaAsset.AssetActions.Stop.Small"),
diff --git a/Engine/Source/Developer/AssetTools/Private/AssetTypeActions/AssetTypeActions_ForceFeedbackEffect.h b/Engine/Source/Developer/AssetTools/Private/AssetTypeActions/AssetTypeActions_ForceFeedbackEffect.h
--- a/Engine/Source/Developer/AssetTools/Private/AssetTypeActions/AssetTypeActions_ForceFeedbackEffect.h
+++ b/Engine/Source/Developer/AssetTools/Private/AssetTypeActions/AssetTypeActions_ForceFeedbackEffect.h
@@ -10,7 +10,6 @@
 #include "TickableEditorObject.h"
 
 struct FAssetData;
-class FMenuBuilder;
 
 struct FPreviewForceFeedbackEffect : public FActiveForceFeedbackEffect, public FTickableEditorObject, public FGCObject
 {
@@ -32,7 +31,7 @@ public:
 	virtual UClass* GetSupportedClass() const override;
 	virtual uint32 GetCategories() override { return EAssetTypeCategories::Misc; }
 	virtual bool HasActions ( const TArray<UObject*>& InObjects ) const override { return true; }
-	virtual void GetActions( const TArray<UObject*>& InObjects, FMenuBuilder& MenuBuilder ) override;
+	virtual void GetActions(const TArray<UObject*>& InObjects, struct FEditorMenuSection& Section) override;
 	virtual void AssetsActivated( const TArray<UObject*>& InObjects, EAssetTypeActivationMethod::Type ActivationType ) override;
 	virtual TSharedPtr<SWidget> GetThumbnailOverlay(const FAssetData& AssetData) const override;
 
diff --git a/Engine/Source/Developer/AssetTools/Private/AssetTypeActions/AssetTypeActions_MaterialFunction.cpp b/Engine/Source/Developer/AssetTools/Private/AssetTypeActions/AssetTypeActions_MaterialFunction.cpp
--- a/Engine/Source/Developer/AssetTools/Private/AssetTypeActions/AssetTypeActions_MaterialFunction.cpp
+++ b/Engine/Source/Developer/AssetTools/Private/AssetTypeActions/AssetTypeActions_MaterialFunction.cpp
@@ -2,7 +2,7 @@
 
 #include "AssetTypeActions/AssetTypeActions_MaterialFunction.h"
 #include "Factories/MaterialFunctionInstanceFactory.h"
-#include "Framework/MultiBox/MultiBoxBuilder.h"
+#include "EditorMenuSubsystem.h"
 #include "UObject/UObjectHash.h"
 #include "UObject/UObjectIterator.h"
 #include "EditorStyleSet.h"
@@ -24,14 +24,15 @@ const FString LayerCompareString = (TEXT("MaterialLayer"));
 const FString BlendCompareString = (TEXT("MaterialLayerBlend"));
 
 
-void FAssetTypeActions_MaterialFunction::GetActions( const TArray<UObject*>& InObjects, FMenuBuilder& MenuBuilder )
+void FAssetTypeActions_MaterialFunction::GetActions(const TArray<UObject*>& InObjects, FEditorMenuSection& Section)
 {
 	auto Functions = GetTypedWeakObjectPtrs<UMaterialFunctionInterface>(InObjects);
 
 	IMaterialEditorModule& MaterialEditorModule = FModuleManager::LoadModuleChecked<IMaterialEditorModule>("MaterialEditor");
 	if (MaterialEditorModule.MaterialLayersEnabled())
 	{
-		MenuBuilder.AddMenuEntry(
+		Section.AddMenuEntry(
+			"MaterialFunction_NewMFI",
 			GetInstanceText(),
 			LOCTEXT("Material_NewMFITooltip", "Creates a parameterized function using this function as a base."),
 			FSlateIcon(FEditorStyle::GetStyleSetName(), "ClassIcon.MaterialInstanceActor"),
@@ -41,7 +42,8 @@ void FAssetTypeActions_MaterialFunction::GetActions( const TArray<UObject*>& InO
 		);
 	}
 
-	MenuBuilder.AddMenuEntry(
+	Section.AddMenuEntry(
+		"MaterialFunction_FindMaterials",
 		LOCTEXT("MaterialFunction_FindMaterials", "Find Materials Using This"),
 		LOCTEXT("MaterialFunction_FindMaterialsTooltip", "Finds the materials that reference this material function in the content browser."),
 		FSlateIcon(FEditorStyle::GetStyleSetName(), "ContentBrowser.AssetActions.GenericFind"),
diff --git a/Engine/Source/Developer/AssetTools/Private/AssetTypeActions/AssetTypeActions_MaterialFunction.h b/Engine/Source/Developer/AssetTools/Private/AssetTypeActions/AssetTypeActions_MaterialFunction.h
--- a/Engine/Source/Developer/AssetTools/Private/AssetTypeActions/AssetTypeActions_MaterialFunction.h
+++ b/Engine/Source/Developer/AssetTools/Private/AssetTypeActions/AssetTypeActions_MaterialFunction.h
@@ -9,7 +9,6 @@
 #include "Materials/MaterialFunctionMaterialLayer.h"
 #include "Materials/MaterialFunctionMaterialLayerBlend.h"
 
-class FMenuBuilder;
 
 class FAssetTypeActions_MaterialFunction : public FAssetTypeActions_Base
 {
@@ -19,7 +18,7 @@ public:
 	virtual FColor GetTypeColor() const override { return FColor(0,175,175); }
 	virtual UClass* GetSupportedClass() const override { return UMaterialFunction::StaticClass(); }
 	virtual bool HasActions ( const TArray<UObject*>& InObjects ) const override { return true; }
-	virtual void GetActions( const TArray<UObject*>& InObjects, FMenuBuilder& MenuBuilder ) override;
+	virtual void GetActions(const TArray<UObject*>& InObjects, struct FEditorMenuSection& Section) override;
 	virtual void OpenAssetEditor( const TArray<UObject*>& InObjects, TSharedPtr<class IToolkitHost> EditWithinLevelEditor = TSharedPtr<IToolkitHost>() ) override;
 	virtual uint32 GetCategories() override { return EAssetTypeCategories::MaterialsAndTextures; }
 	virtual class UThumbnailInfo* GetThumbnailInfo(UObject* Asset) const override;
diff --git a/Engine/Source/Developer/AssetTools/Private/AssetTypeActions/AssetTypeActions_MaterialFunctionInstance.cpp b/Engine/Source/Developer/AssetTools/Private/AssetTypeActions/AssetTypeActions_MaterialFunctionInstance.cpp
--- a/Engine/Source/Developer/AssetTools/Private/AssetTypeActions/AssetTypeActions_MaterialFunctionInstance.cpp
+++ b/Engine/Source/Developer/AssetTools/Private/AssetTypeActions/AssetTypeActions_MaterialFunctionInstance.cpp
@@ -1,7 +1,7 @@
 // Copyright 1998-2019 Epic Games, Inc. All Rights Reserved.
 
 #include "AssetTypeActions/AssetTypeActions_MaterialFunctionInstance.h"
-#include "Framework/MultiBox/MultiBoxBuilder.h"
+#include "EditorMenuSubsystem.h"
 #include "UObject/UObjectHash.h"
 #include "UObject/UObjectIterator.h"
 #include "EditorStyleSet.h"
@@ -14,13 +14,14 @@
 
 #define LOCTEXT_NAMESPACE "AssetTypeActions"
 
-void FAssetTypeActions_MaterialFunctionInstance::GetActions( const TArray<UObject*>& InObjects, FMenuBuilder& MenuBuilder )
+void FAssetTypeActions_MaterialFunctionInstance::GetActions(const TArray<UObject*>& InObjects, FEditorMenuSection& Section)
 {
 	auto MFIs = GetTypedWeakObjectPtrs<UMaterialFunctionInstance>(InObjects);
 
-	FAssetTypeActions_MaterialFunction::GetActions(InObjects, MenuBuilder);
+	FAssetTypeActions_MaterialFunction::GetActions(InObjects, Section);
 
-	MenuBuilder.AddMenuEntry(
+	Section.AddMenuEntry(
+		"MaterialFunctionInstance_FindParent",
 		LOCTEXT("MaterialFunctionInstance_FindParent", "Find Parent"),
 		LOCTEXT("MaterialFunctionInstance_FindParentTooltip", "Finds the function this instance is based on in the content browser."),
 		FSlateIcon(FEditorStyle::GetStyleSetName(), "ContentBrowser.AssetActions.GenericFind"),
diff --git a/Engine/Source/Developer/AssetTools/Private/AssetTypeActions/AssetTypeActions_MaterialFunctionInstance.h b/Engine/Source/Developer/AssetTools/Private/AssetTypeActions/AssetTypeActions_MaterialFunctionInstance.h
--- a/Engine/Source/Developer/AssetTools/Private/AssetTypeActions/AssetTypeActions_MaterialFunctionInstance.h
+++ b/Engine/Source/Developer/AssetTools/Private/AssetTypeActions/AssetTypeActions_MaterialFunctionInstance.h
@@ -6,7 +6,6 @@
 #include "AssetTypeActions_MaterialFunction.h"
 #include "Materials/MaterialFunctionInstance.h"
 
-class FMenuBuilder;
 
 class FAssetTypeActions_MaterialFunctionInstance : public FAssetTypeActions_MaterialFunction
 {
@@ -15,7 +14,7 @@ public:
 	virtual FText GetName() const override { return NSLOCTEXT("AssetTypeActions", "FAssetTypeActions_MaterialFunctionInstance", "Material Function Instance"); }
 	virtual UClass* GetSupportedClass() const override { return UMaterialFunctionInstance::StaticClass(); }
 	virtual uint32 GetCategories() override { return EAssetTypeCategories::None; }
-	virtual void GetActions( const TArray<UObject*>& InObjects, FMenuBuilder& MenuBuilder ) override;
+	virtual void GetActions(const TArray<UObject*>& InObjects, struct FEditorMenuSection& Section) override;
 	virtual void OpenAssetEditor( const TArray<UObject*>& InObjects, TSharedPtr<class IToolkitHost> EditWithinLevelEditor = TSharedPtr<IToolkitHost>() ) override;
 	virtual class UThumbnailInfo* GetThumbnailInfo(UObject* Asset) const override;
 
@@ -51,4 +50,4 @@ private:
 	virtual void ExecuteNewMFI(TArray<TWeakObjectPtr<UMaterialFunctionInterface>> Objects) override;
 	virtual FText GetInstanceText() const override { return NSLOCTEXT("AssetTypeActions", "Material_NewMBI", "Create Blend Instance"); }
 
-};
\ No newline at end of file
+};
diff --git a/Engine/Source/Developer/AssetTools/Private/AssetTypeActions/AssetTypeActions_MaterialInstanceConstant.cpp b/Engine/Source/Developer/AssetTools/Private/AssetTypeActions/AssetTypeActions_MaterialInstanceConstant.cpp
--- a/Engine/Source/Developer/AssetTools/Private/AssetTypeActions/AssetTypeActions_MaterialInstanceConstant.cpp
+++ b/Engine/Source/Developer/AssetTools/Private/AssetTypeActions/AssetTypeActions_MaterialInstanceConstant.cpp
@@ -1,20 +1,21 @@
 // Copyright 1998-2019 Epic Games, Inc. All Rights Reserved.
 
 #include "AssetTypeActions/AssetTypeActions_MaterialInstanceConstant.h"
-#include "Framework/MultiBox/MultiBoxBuilder.h"
+#include "EditorMenuSubsystem.h"
 #include "EditorStyleSet.h"
 #include "AssetTools.h"
 #include "MaterialEditorModule.h"
 
 #define LOCTEXT_NAMESPACE "AssetTypeActions"
 
-void FAssetTypeActions_MaterialInstanceConstant::GetActions( const TArray<UObject*>& InObjects, FMenuBuilder& MenuBuilder )
+void FAssetTypeActions_MaterialInstanceConstant::GetActions(const TArray<UObject*>& InObjects, FEditorMenuSection& Section)
 {
 	auto MICs = GetTypedWeakObjectPtrs<UMaterialInstanceConstant>(InObjects);
 
-	FAssetTypeActions_MaterialInterface::GetActions(InObjects, MenuBuilder);
+	FAssetTypeActions_MaterialInterface::GetActions(InObjects, Section);
 
-	MenuBuilder.AddMenuEntry(
+	Section.AddMenuEntry(
+		"MaterialInstanceConstant_FindParent",
 		LOCTEXT("MaterialInstanceConstant_FindParent", "Find Parent"),
 		LOCTEXT("MaterialInstanceConstant_FindParentTooltip", "Finds the material this instance is based on in the content browser."),
 		FSlateIcon(FEditorStyle::GetStyleSetName(), "ContentBrowser.AssetActions.GenericFind"),
diff --git a/Engine/Source/Developer/AssetTools/Private/AssetTypeActions/AssetTypeActions_MaterialInstanceConstant.h b/Engine/Source/Developer/AssetTools/Private/AssetTypeActions/AssetTypeActions_MaterialInstanceConstant.h
--- a/Engine/Source/Developer/AssetTools/Private/AssetTypeActions/AssetTypeActions_MaterialInstanceConstant.h
+++ b/Engine/Source/Developer/AssetTools/Private/AssetTypeActions/AssetTypeActions_MaterialInstanceConstant.h
@@ -7,7 +7,6 @@
 #include "AssetTypeActions/AssetTypeActions_MaterialInterface.h"
 #include "Materials/MaterialInstanceConstant.h"
 
-class FMenuBuilder;
 
 class FAssetTypeActions_MaterialInstanceConstant : public FAssetTypeActions_MaterialInterface
 {
@@ -21,7 +20,7 @@ public:
 	virtual FText GetName() const override { return NSLOCTEXT("AssetTypeActions", "AssetTypeActions_MaterialInstanceConstant", "Material Instance"); }
 	virtual FColor GetTypeColor() const override { return FColor(0,128,0); }
 	virtual UClass* GetSupportedClass() const override { return UMaterialInstanceConstant::StaticClass(); }
-	virtual void GetActions( const TArray<UObject*>& InObjects, FMenuBuilder& MenuBuilder ) override;
+	virtual void GetActions(const TArray<UObject*>& InObjects, struct FEditorMenuSection& Section) override;
 	virtual void OpenAssetEditor( const TArray<UObject*>& InObjects, TSharedPtr<class IToolkitHost> EditWithinLevelEditor = TSharedPtr<IToolkitHost>() ) override;
 	virtual bool CanFilter() override { return true; }
 	virtual uint32 GetCategories() override { return FAssetTypeActions_MaterialInterface::GetCategories() | AssetCategoryBit; }
diff --git a/Engine/Source/Developer/AssetTools/Private/AssetTypeActions/AssetTypeActions_MaterialInterface.cpp b/Engine/Source/Developer/AssetTools/Private/AssetTypeActions/AssetTypeActions_MaterialInterface.cpp
--- a/Engine/Source/Developer/AssetTools/Private/AssetTypeActions/AssetTypeActions_MaterialInterface.cpp
+++ b/Engine/Source/Developer/AssetTools/Private/AssetTypeActions/AssetTypeActions_MaterialInterface.cpp
@@ -2,7 +2,7 @@
 
 #include "AssetTypeActions/AssetTypeActions_MaterialInterface.h"
 #include "Misc/PackageName.h"
-#include "Framework/MultiBox/MultiBoxBuilder.h"
+#include "EditorMenuSubsystem.h"
 #include "EditorStyleSet.h"
 #include "Factories/MaterialInstanceConstantFactoryNew.h"
 #include "ThumbnailRendering/SceneThumbnailInfoWithPrimitive.h"
@@ -13,11 +13,12 @@
 
 #define LOCTEXT_NAMESPACE "AssetTypeActions"
 
-void FAssetTypeActions_MaterialInterface::GetActions( const TArray<UObject*>& InObjects, FMenuBuilder& MenuBuilder )
+void FAssetTypeActions_MaterialInterface::GetActions(const TArray<UObject*>& InObjects, FEditorMenuSection& Section)
 {
 	auto MaterialInterfaces = GetTypedWeakObjectPtrs<UMaterialInterface>(InObjects);
 
-	MenuBuilder.AddMenuEntry(
+	Section.AddMenuEntry(
+		"Material_NewMIC",
 		LOCTEXT("Material_NewMIC", "Create Material Instance"),
 		LOCTEXT("Material_NewMICTooltip", "Creates a parameterized material using this material as a base."),
 		FSlateIcon(FEditorStyle::GetStyleSetName(), "ClassIcon.MaterialInstanceActor"),
diff --git a/Engine/Source/Developer/AssetTools/Private/AssetTypeActions/AssetTypeActions_MaterialInterface.h b/Engine/Source/Developer/AssetTools/Private/AssetTypeActions/AssetTypeActions_MaterialInterface.h
--- a/Engine/Source/Developer/AssetTools/Private/AssetTypeActions/AssetTypeActions_MaterialInterface.h
+++ b/Engine/Source/Developer/AssetTools/Private/AssetTypeActions/AssetTypeActions_MaterialInterface.h
@@ -6,7 +6,6 @@
 #include "AssetTypeActions_Base.h"
 #include "Materials/MaterialInterface.h"
 
-class FMenuBuilder;
 
 class FAssetTypeActions_MaterialInterface : public FAssetTypeActions_Base
 {
@@ -16,7 +15,7 @@ public:
 	virtual FColor GetTypeColor() const override { return FColor(64,192,64); }
 	virtual UClass* GetSupportedClass() const override { return UMaterialInterface::StaticClass(); }
 	virtual bool HasActions ( const TArray<UObject*>& InObjects ) const override { return true; }
-	virtual void GetActions( const TArray<UObject*>& InObjects, FMenuBuilder& MenuBuilder ) override;
+	virtual void GetActions(const TArray<UObject*>& InObjects, struct FEditorMenuSection& Section) override;
 	virtual bool CanFilter() override { return false; }
 	virtual uint32 GetCategories() override { return EAssetTypeCategories::MaterialsAndTextures; }
 	virtual class UThumbnailInfo* GetThumbnailInfo(UObject* Asset) const override;
diff --git a/Engine/Source/Developer/AssetTools/Private/AssetTypeActions/AssetTypeActions_ParticleSystem.cpp b/Engine/Source/Developer/AssetTools/Private/AssetTypeActions/AssetTypeActions_ParticleSystem.cpp
--- a/Engine/Source/Developer/AssetTools/Private/AssetTypeActions/AssetTypeActions_ParticleSystem.cpp
+++ b/Engine/Source/Developer/AssetTools/Private/AssetTypeActions/AssetTypeActions_ParticleSystem.cpp
@@ -1,7 +1,7 @@
 // Copyright 1998-2019 Epic Games, Inc. All Rights Reserved.
 
 #include "AssetTypeActions/AssetTypeActions_ParticleSystem.h"
-#include "Framework/MultiBox/MultiBoxBuilder.h"
+#include "EditorMenuSubsystem.h"
 #include "Misc/FeedbackContext.h"
 #include "Editor/Cascade/Public/CascadeModule.h"
 #include "Particles/ParticleEmitter.h"
@@ -9,11 +9,12 @@
 
 #define LOCTEXT_NAMESPACE "AssetTypeActions"
 
-void FAssetTypeActions_ParticleSystem::GetActions( const TArray<UObject*>& InObjects, FMenuBuilder& MenuBuilder )
+void FAssetTypeActions_ParticleSystem::GetActions(const TArray<UObject*>& InObjects, FEditorMenuSection& Section)
 {
 	auto ParticleSystems = GetTypedWeakObjectPtrs<UParticleSystem>(InObjects);
 
-	MenuBuilder.AddMenuEntry(
+	Section.AddMenuEntry(
+		"ParticleSystem_CopyParameters",
 		LOCTEXT("ParticleSystem_CopyParameters", "Copy Parameters"),
 		LOCTEXT("ParticleSystem_CopyParametersTooltip", "Copies particle system parameters to the clipboard."),
 		FSlateIcon(),
@@ -23,7 +24,8 @@ void FAssetTypeActions_ParticleSystem::GetActions( const TArray<UObject*>& InObj
 			)
 		);
 
-	MenuBuilder.AddMenuEntry(
+	Section.AddMenuEntry(
+		"ParticleSystem_ConvertToSeeded",
 		LOCTEXT("ParticleSystem_ConvertToSeeded", "Convert To Seeded"),
 		LOCTEXT("ParticleSystem_ConvertToSeededTooltip", "Convert all modules in this particle system to random seeded modules"),
 		FSlateIcon(),
diff --git a/Engine/Source/Developer/AssetTools/Private/AssetTypeActions/AssetTypeActions_ParticleSystem.h b/Engine/Source/Developer/AssetTools/Private/AssetTypeActions/AssetTypeActions_ParticleSystem.h
--- a/Engine/Source/Developer/AssetTools/Private/AssetTypeActions/AssetTypeActions_ParticleSystem.h
+++ b/Engine/Source/Developer/AssetTools/Private/AssetTypeActions/AssetTypeActions_ParticleSystem.h
@@ -7,7 +7,6 @@
 #include "AssetTypeActions_Base.h"
 #include "Particles/ParticleSystem.h"
 
-class FMenuBuilder;
 
 class FAssetTypeActions_ParticleSystem : public FAssetTypeActions_Base
 {
@@ -17,7 +16,7 @@ public:
 	virtual FColor GetTypeColor() const override { return FColor(255,255,255); }
 	virtual UClass* GetSupportedClass() const override { return UParticleSystem::StaticClass(); }
 	virtual bool HasActions ( const TArray<UObject*>& InObjects ) const override { return true; }
-	virtual void GetActions( const TArray<UObject*>& InObjects, FMenuBuilder& MenuBuilder ) override;
+	virtual void GetActions(const TArray<UObject*>& InObjects, struct FEditorMenuSection& Section) override;
 	virtual void OpenAssetEditor( const TArray<UObject*>& InObjects, TSharedPtr<class IToolkitHost> EditWithinLevelEditor = TSharedPtr<IToolkitHost>() ) override;
 	virtual uint32 GetCategories() override { return EAssetTypeCategories::Basic; }
 	
diff --git a/Engine/Source/Developer/AssetTools/Private/AssetTypeActions/AssetTypeActions_Redirector.cpp b/Engine/Source/Developer/AssetTools/Private/AssetTypeActions/AssetTypeActions_Redirector.cpp
--- a/Engine/Source/Developer/AssetTools/Private/AssetTypeActions/AssetTypeActions_Redirector.cpp
+++ b/Engine/Source/Developer/AssetTools/Private/AssetTypeActions/AssetTypeActions_Redirector.cpp
@@ -1,16 +1,17 @@
 // Copyright 1998-2019 Epic Games, Inc. All Rights Reserved.
 
 #include "AssetTypeActions/AssetTypeActions_Redirector.h"
-#include "Framework/MultiBox/MultiBoxBuilder.h"
+#include "EditorMenuSubsystem.h"
 #include "AssetTools.h"
 
 #define LOCTEXT_NAMESPACE "AssetTypeActions"
 
-void FAssetTypeActions_Redirector::GetActions( const TArray<UObject*>& InObjects, FMenuBuilder& MenuBuilder )
+void FAssetTypeActions_Redirector::GetActions(const TArray<UObject*>& InObjects, FEditorMenuSection& Section)
 {
 	auto Redirectors = GetTypedWeakObjectPtrs<UObjectRedirector>(InObjects);
 
-	MenuBuilder.AddMenuEntry(
+	Section.AddMenuEntry(
+		"Redirector_FindTarget",
 		LOCTEXT("Redirector_FindTarget","Find Target"),
 		LOCTEXT("Redirector_FindTargetTooltip", "Finds the asset that this redirector targets in the asset tree."),
 		FSlateIcon(),
@@ -20,7 +21,8 @@ void FAssetTypeActions_Redirector::GetActions( const TArray<UObject*>& InObjects
 			)
 		);
 
-	MenuBuilder.AddMenuEntry(
+	Section.AddMenuEntry(
+		"Redirector_FixUp",
 		LOCTEXT("Redirector_FixUp","Fix Up"),
 		LOCTEXT("Redirector_FixUpTooltip", "Finds referencers to selected redirectors and resaves them if possible, then deletes any redirectors that had all their referencers fixed."),
 		FSlateIcon(),
diff --git a/Engine/Source/Developer/AssetTools/Private/AssetTypeActions/AssetTypeActions_Redirector.h b/Engine/Source/Developer/AssetTools/Private/AssetTypeActions/AssetTypeActions_Redirector.h
--- a/Engine/Source/Developer/AssetTools/Private/AssetTypeActions/AssetTypeActions_Redirector.h
+++ b/Engine/Source/Developer/AssetTools/Private/AssetTypeActions/AssetTypeActions_Redirector.h
@@ -6,7 +6,6 @@
 #include "UObject/ObjectRedirector.h"
 #include "AssetTypeActions_Base.h"
 
-class FMenuBuilder;
 
 class FAssetTypeActions_Redirector : public FAssetTypeActions_Base
 {
@@ -16,7 +15,7 @@ public:
 	virtual FColor GetTypeColor() const override { return FColor(128, 128, 128); }
 	virtual UClass* GetSupportedClass() const override { return UObjectRedirector::StaticClass(); }
 	virtual bool HasActions ( const TArray<UObject*>& InObjects ) const override { return true; }
-	virtual void GetActions( const TArray<UObject*>& InObjects, FMenuBuilder& MenuBuilder ) override;
+	virtual void GetActions(const TArray<UObject*>& InObjects, struct FEditorMenuSection& Section) override;
 	virtual void AssetsActivated( const TArray<UObject*>& InObjects, EAssetTypeActivationMethod::Type ActivationType ) override;
 	virtual uint32 GetCategories() override { return EAssetTypeCategories::Misc; }
 	virtual bool CanLocalize() const override { return false; }
diff --git a/Engine/Source/Developer/AssetTools/Private/AssetTypeActions/AssetTypeActions_SkeletalMesh.cpp b/Engine/Source/Developer/AssetTools/Private/AssetTypeActions/AssetTypeActions_SkeletalMesh.cpp
--- a/Engine/Source/Developer/AssetTools/Private/AssetTypeActions/AssetTypeActions_SkeletalMesh.cpp
+++ b/Engine/Source/Developer/AssetTools/Private/AssetTypeActions/AssetTypeActions_SkeletalMesh.cpp
@@ -3,7 +3,7 @@
 #include "AssetTypeActions/AssetTypeActions_SkeletalMesh.h"
 #include "Animation/Skeleton.h"
 #include "Widgets/Text/STextBlock.h"
-#include "Framework/MultiBox/MultiBoxBuilder.h"
+#include "EditorMenuSubsystem.h"
 #include "Misc/PackageName.h"
 #include "Widgets/Layout/SBorder.h"
 #include "Misc/MessageDialog.h"
@@ -448,36 +448,43 @@ FAssetTypeActions_SkeletalMesh::FAssetTypeActions_SkeletalMesh()
 	AssetRegistryModule.Get().OnAssetRemoved().AddRaw(this, &FAssetTypeActions_SkeletalMesh::OnAssetRemoved);
 }
 
-void FAssetTypeActions_SkeletalMesh::GetActions( const TArray<UObject*>& InObjects, FMenuBuilder& MenuBuilder )
+void FAssetTypeActions_SkeletalMesh::GetActions(const TArray<UObject*>& InObjects, FEditorMenuSection& Section)
 {
 	auto Meshes = GetTypedWeakObjectPtrs<USkeletalMesh>(InObjects);
 
-	MenuBuilder.AddSubMenu(
+	Section.AddEntry(FEditorMenuEntry::InitSubMenu(
+			NAME_None,
+			"CreateSkeletalMeshSubmenu",
 			LOCTEXT("CreateSkeletalMeshSubmenu", "Create"),
 			LOCTEXT("CreateSkeletalMeshSubmenu_ToolTip", "Create related assets"),
 			FNewMenuDelegate::CreateSP(this, &FAssetTypeActions_SkeletalMesh::FillCreateMenu, Meshes),
 			false,
 			FSlateIcon(FEditorStyle::GetStyleSetName(), "Persona.AssetActions.CreateAnimAsset")
-			);
+			));
 
-	MenuBuilder.AddSubMenu(	
+	Section.AddEntry(FEditorMenuEntry::InitSubMenu(
+		NAME_None,
+		"SkeletalMesh_LODImport",	
 		LOCTEXT("SkeletalMesh_LODImport", "Import LOD"),
 		LOCTEXT("SkeletalMesh_LODImportTooltip", "Select which LODs to import."),
 		FNewMenuDelegate::CreateSP(this, &FAssetTypeActions_SkeletalMesh::GetLODMenu, Meshes)
-		);
+		));
 	
-	MenuBuilder.AddMenuEntry(
+	Section.AddMenuEntry(
+		"ImportClothing_Entry",
 		LOCTEXT("ImportClothing_Entry", "Import Clothing Asset..."),
 		LOCTEXT("ImportClothing_ToolTip", "Import a clothing asset from a supported file on disk into this skeletal mesh."),
 		FSlateIcon(),
 		FUIAction(FExecuteAction::CreateSP(this, &FAssetTypeActions_SkeletalMesh::ExecuteImportClothing, Meshes)));
 
 	// skeleton menu
-	MenuBuilder.AddSubMenu(
+	Section.AddEntry(FEditorMenuEntry::InitSubMenu(
+		NAME_None,
+		"SkeletonSubmenu",
 		LOCTEXT("SkeletonSubmenu", "Skeleton"),
 		LOCTEXT("SkeletonSubmenu_ToolTip", "Skeleton related actions"),
 		FNewMenuDelegate::CreateSP(this, &FAssetTypeActions_SkeletalMesh::FillSkeletonMenu, Meshes)
-	);
+	));
 }
 
 void FAssetTypeActions_SkeletalMesh::FillCreateMenu(FMenuBuilder& MenuBuilder, TArray<TWeakObjectPtr<USkeletalMesh>> Meshes) const
diff --git a/Engine/Source/Developer/AssetTools/Private/AssetTypeActions/AssetTypeActions_SkeletalMesh.h b/Engine/Source/Developer/AssetTools/Private/AssetTypeActions/AssetTypeActions_SkeletalMesh.h
--- a/Engine/Source/Developer/AssetTools/Private/AssetTypeActions/AssetTypeActions_SkeletalMesh.h
+++ b/Engine/Source/Developer/AssetTools/Private/AssetTypeActions/AssetTypeActions_SkeletalMesh.h
@@ -7,7 +7,6 @@
 #include "AssetTypeActions_Base.h"
 #include "Engine/SkeletalMesh.h"
 
-class FMenuBuilder;
 
 class FAssetTypeActions_SkeletalMesh : public FAssetTypeActions_Base
 {
@@ -19,7 +18,7 @@ public:
 	virtual FColor GetTypeColor() const override { return FColor(255,0,255); }
 	virtual UClass* GetSupportedClass() const override { return USkeletalMesh::StaticClass(); }
 	virtual bool HasActions ( const TArray<UObject*>& InObjects ) const override { return true; }
-	virtual void GetActions( const TArray<UObject*>& InObjects, FMenuBuilder& MenuBuilder ) override;
+	virtual void GetActions(const TArray<UObject*>& InObjects, struct FEditorMenuSection& Section) override;
 	virtual void OpenAssetEditor( const TArray<UObject*>& InObjects, TSharedPtr<class IToolkitHost> EditWithinLevelEditor = TSharedPtr<IToolkitHost>() ) override;
 	virtual uint32 GetCategories() override { return EAssetTypeCategories::Basic | EAssetTypeCategories::Animation; }
 	virtual class UThumbnailInfo* GetThumbnailInfo(UObject* Asset) const override;
diff --git a/Engine/Source/Developer/AssetTools/Private/AssetTypeActions/AssetTypeActions_Skeleton.cpp b/Engine/Source/Developer/AssetTools/Private/AssetTypeActions/AssetTypeActions_Skeleton.cpp
--- a/Engine/Source/Developer/AssetTools/Private/AssetTypeActions/AssetTypeActions_Skeleton.cpp
+++ b/Engine/Source/Developer/AssetTools/Private/AssetTypeActions/AssetTypeActions_Skeleton.cpp
@@ -2,7 +2,7 @@
 
 #include "AssetTypeActions/AssetTypeActions_Skeleton.h"
 #include "Widgets/Text/STextBlock.h"
-#include "Framework/MultiBox/MultiBoxBuilder.h"
+#include "EditorMenuSubsystem.h"
 #include "Engine/SkeletalMesh.h"
 #include "Animation/AnimationAsset.h"
 #include "Animation/AnimSequenceBase.h"
@@ -425,23 +425,26 @@ FCreateRigDlg::EResult FCreateRigDlg::ShowModal()
 }
 
 ///////////////////////////////
-void FAssetTypeActions_Skeleton::GetActions( const TArray<UObject*>& InObjects, FMenuBuilder& MenuBuilder )
+void FAssetTypeActions_Skeleton::GetActions(const TArray<UObject*>& InObjects, FEditorMenuSection& Section)
 {
 	auto Skeletons = GetTypedWeakObjectPtrs<USkeleton>(InObjects);
 
 	// create menu
-	MenuBuilder.AddSubMenu(
+	Section.AddEntry(FEditorMenuEntry::InitSubMenu(
+			NAME_None,
+			"CreateSkeletonSubmenu",
 			LOCTEXT("CreateSkeletonSubmenu", "Create"),
 			LOCTEXT("CreateSkeletonSubmenu_ToolTip", "Create assets for this skeleton"),
 			FNewMenuDelegate::CreateSP(this, &FAssetTypeActions_Skeleton::FillCreateMenu, Skeletons),
 			false, 
 			FSlateIcon(FEditorStyle::GetStyleSetName(), "Persona.AssetActions.CreateAnimAsset")
-			);
+			));
 
 	// only show if one is selected. It won't work since I changed the window to be normal window
 	if (Skeletons.Num() == 1)
 	{
-		MenuBuilder.AddMenuEntry(
+		Section.AddMenuEntry(
+			"Skeleton_Retarget",
 			LOCTEXT("Skeleton_Retarget", "Retarget to Another Skeleton"),
 			LOCTEXT("Skeleton_RetargetTooltip", "Allow all animation assets for this skeleton retarget to another skeleton."),
 			FSlateIcon(FEditorStyle::GetStyleSetName(), "Persona.AssetActions.RetargetSkeleton"),
@@ -454,7 +457,8 @@ void FAssetTypeActions_Skeleton::GetActions( const TArray<UObject*>& InObjects,
 
 	// @todo ImportAnimation
 	/*
-	MenuBuilder.AddMenuEntry(
+	Section.AddMenuEntry(
+		"Skeleton_ImportAnimation",
 		LOCTEXT("Skeleton_ImportAnimation", "Import Animation"),
 		LOCTEXT("Skeleton_ImportAnimationTooltip", "Imports an animation for the selected skeleton."),
 		FSlateIcon(),
@@ -463,7 +467,7 @@ void FAssetTypeActions_Skeleton::GetActions( const TArray<UObject*>& InObjects,
 			FCanExecuteAction()
 			)
 		);
-	*/
+	*/;
 }
 
 void FAssetTypeActions_Skeleton::FillCreateMenu(FMenuBuilder& MenuBuilder, TArray<TWeakObjectPtr<USkeleton>> Skeletons) const
diff --git a/Engine/Source/Developer/AssetTools/Private/AssetTypeActions/AssetTypeActions_Skeleton.h b/Engine/Source/Developer/AssetTools/Private/AssetTypeActions/AssetTypeActions_Skeleton.h
--- a/Engine/Source/Developer/AssetTools/Private/AssetTypeActions/AssetTypeActions_Skeleton.h
+++ b/Engine/Source/Developer/AssetTools/Private/AssetTypeActions/AssetTypeActions_Skeleton.h
@@ -8,7 +8,6 @@
 #include "AssetTypeActions_Base.h"
 #include "EditorAnimUtils.h"
 
-class FMenuBuilder;
 
 /**
  * Remap Skeleton Asset Data
@@ -41,7 +40,7 @@ public:
 	virtual FColor GetTypeColor() const override { return FColor(105,181,205); }
 	virtual UClass* GetSupportedClass() const override { return USkeleton::StaticClass(); }
 	virtual bool HasActions ( const TArray<UObject*>& InObjects ) const override { return true; }
-	virtual void GetActions( const TArray<UObject*>& InObjects, FMenuBuilder& MenuBuilder ) override;
+	virtual void GetActions(const TArray<UObject*>& InObjects, struct FEditorMenuSection& Section) override;
 	virtual void OpenAssetEditor( const TArray<UObject*>& InObjects, TSharedPtr<class IToolkitHost> EditWithinLevelEditor = TSharedPtr<IToolkitHost>() ) override;
 	virtual uint32 GetCategories() override { return EAssetTypeCategories::Animation; }
 
diff --git a/Engine/Source/Developer/AssetTools/Private/AssetTypeActions/AssetTypeActions_SlateWidgetStyle.cpp b/Engine/Source/Developer/AssetTools/Private/AssetTypeActions/AssetTypeActions_SlateWidgetStyle.cpp
--- a/Engine/Source/Developer/AssetTools/Private/AssetTypeActions/AssetTypeActions_SlateWidgetStyle.cpp
+++ b/Engine/Source/Developer/AssetTools/Private/AssetTypeActions/AssetTypeActions_SlateWidgetStyle.cpp
@@ -1,12 +1,13 @@
 // Copyright 1998-2019 Epic Games, Inc. All Rights Reserved.
 
 #include "AssetTypeActions/AssetTypeActions_SlateWidgetStyle.h"
+#include "EditorMenuSubsystem.h"
 
 
 #define LOCTEXT_NAMESPACE "AssetTypeActions"
 
 
-void FAssetTypeActions_SlateWidgetStyle::GetActions( const TArray<UObject*>& InObjects, FMenuBuilder& MenuBuilder )
+void FAssetTypeActions_SlateWidgetStyle::GetActions(const TArray<UObject*>& InObjects, FEditorMenuSection& Section)
 {
 	auto Styles = GetTypedWeakObjectPtrs< USlateWidgetStyleAsset >(InObjects);
 }
diff --git a/Engine/Source/Developer/AssetTools/Private/AssetTypeActions/AssetTypeActions_SlateWidgetStyle.h b/Engine/Source/Developer/AssetTools/Private/AssetTypeActions/AssetTypeActions_SlateWidgetStyle.h
--- a/Engine/Source/Developer/AssetTools/Private/AssetTypeActions/AssetTypeActions_SlateWidgetStyle.h
+++ b/Engine/Source/Developer/AssetTools/Private/AssetTypeActions/AssetTypeActions_SlateWidgetStyle.h
@@ -7,7 +7,6 @@
 #include "Toolkits/IToolkitHost.h"
 #include "AssetTypeActions_Base.h"
 
-class FMenuBuilder;
 
 class FAssetTypeActions_SlateWidgetStyle : public FAssetTypeActions_Base
 {
@@ -18,7 +17,7 @@ public:
 	virtual UClass* GetSupportedClass() const override { return USlateWidgetStyleAsset::StaticClass(); }
 	virtual uint32 GetCategories() override { return EAssetTypeCategories::UI; }
 	virtual bool HasActions ( const TArray<UObject*>& InObjects ) const override { return true; }
-	virtual void GetActions( const TArray<UObject*>& InObjects, FMenuBuilder& MenuBuilder ) override;
+	virtual void GetActions(const TArray<UObject*>& InObjects, struct FEditorMenuSection& Section) override;
 
 	virtual void OpenAssetEditor( const TArray<UObject*>& InObjects, TSharedPtr<class IToolkitHost> EditWithinLevelEditor = TSharedPtr<IToolkitHost>() ) override;
 };
diff --git a/Engine/Source/Developer/AssetTools/Private/AssetTypeActions/AssetTypeActions_Texture.cpp b/Engine/Source/Developer/AssetTools/Private/AssetTypeActions/AssetTypeActions_Texture.cpp
--- a/Engine/Source/Developer/AssetTools/Private/AssetTypeActions/AssetTypeActions_Texture.cpp
+++ b/Engine/Source/Developer/AssetTools/Private/AssetTypeActions/AssetTypeActions_Texture.cpp
@@ -21,7 +21,7 @@
 #include "Editor.h"
 #include "Widgets/Text/STextBlock.h"
 #include "Widgets/Images/SImage.h"
-#include "Framework/MultiBox/MultiBoxBuilder.h"
+#include "EditorMenuSubsystem.h"
 #include "Widgets/Layout/SBorder.h"
 #include "Misc/MessageDialog.h"
 #include "Misc/FeedbackContext.h"
@@ -44,7 +44,7 @@
 
 #define LOCTEXT_NAMESPACE "AssetTypeActions"
 
-void FAssetTypeActions_Texture::GetActions( const TArray<UObject*>& InObjects, FMenuBuilder& MenuBuilder )
+void FAssetTypeActions_Texture::GetActions(const TArray<UObject*>& InObjects, FEditorMenuSection& Section)
 {
 	auto Textures = GetTypedWeakObjectPtrs<UTexture>(InObjects);
 	bool bHasVirtualTextures = false;
@@ -62,7 +62,8 @@ void FAssetTypeActions_Texture::GetActions( const TArray<UObject*>& InObjects, F
 		}
 	}
 
-	MenuBuilder.AddMenuEntry(
+	Section.AddMenuEntry(
+		"Texture_CreateMaterial",
 		LOCTEXT("Texture_CreateMaterial", "Create Material"),
 		LOCTEXT("Texture_CreateMaterialTooltip", "Creates a new material using this texture."),
 		FSlateIcon(FEditorStyle::GetStyleSetName(), "ClassIcon.Material"),
@@ -74,7 +75,8 @@ void FAssetTypeActions_Texture::GetActions( const TArray<UObject*>& InObjects, F
 
 	if (bHasNonVirtualTextures)
 	{
-		MenuBuilder.AddMenuEntry(
+		Section.AddMenuEntry(
+			"Texture_ConvertToVT",
 			LOCTEXT("Texture_ConvertToVT", "Convert to Virtual Texture"),
 			LOCTEXT("Texture_ConvertToVTTooltip", "Converts this texture to a virtual texture if it fits the size limit imposed in the texture importer settings."),
 			FSlateIcon(FEditorStyle::GetStyleSetName(), "ClassIcon.Texture2D"),
@@ -87,7 +89,8 @@ void FAssetTypeActions_Texture::GetActions( const TArray<UObject*>& InObjects, F
 
 	if (bHasVirtualTextures)
 	{
-		MenuBuilder.AddMenuEntry(
+		Section.AddMenuEntry(
+			"Texture_ConvertToRegular",
 			LOCTEXT("Texture_ConvertToRegular", "Convert to Regular Texture"),
 			LOCTEXT("Texture_ConvertToRegularTooltip", "Converts this texture to a regular 2D texture if it is a virtual texture."),
 			FSlateIcon(FEditorStyle::GetStyleSetName(), "ClassIcon.Texture2D"),
@@ -100,7 +103,8 @@ void FAssetTypeActions_Texture::GetActions( const TArray<UObject*>& InObjects, F
 
 	if ( InObjects.Num() == 1 )
 	{
-		MenuBuilder.AddMenuEntry(
+		Section.AddMenuEntry(
+			"Texture_FindMaterials",
 			LOCTEXT("Texture_FindMaterials", "Find Materials Using This"),
 			LOCTEXT("Texture_FindMaterialsTooltip", "Finds all materials that use this material in the content browser."),
 			FSlateIcon(),
diff --git a/Engine/Source/Developer/AssetTools/Private/AssetTypeActions/AssetTypeActions_Texture.h b/Engine/Source/Developer/AssetTools/Private/AssetTypeActions/AssetTypeActions_Texture.h
--- a/Engine/Source/Developer/AssetTools/Private/AssetTypeActions/AssetTypeActions_Texture.h
+++ b/Engine/Source/Developer/AssetTools/Private/AssetTypeActions/AssetTypeActions_Texture.h
@@ -7,7 +7,6 @@
 #include "Toolkits/IToolkitHost.h"
 #include "AssetTypeActions_Base.h"
 
-class FMenuBuilder;
 
 class FAssetTypeActions_Texture : public FAssetTypeActions_Base
 {
@@ -17,7 +16,7 @@ public:
 	virtual FColor GetTypeColor() const override { return FColor(192,64,64); }
 	virtual UClass* GetSupportedClass() const override { return UTexture::StaticClass(); }
 	virtual bool HasActions ( const TArray<UObject*>& InObjects ) const override { return true; }
-	virtual void GetActions( const TArray<UObject*>& InObjects, FMenuBuilder& MenuBuilder ) override;
+	virtual void GetActions(const TArray<UObject*>& InObjects, struct FEditorMenuSection& Section) override;
 	virtual void OpenAssetEditor( const TArray<UObject*>& InObjects, TSharedPtr<class IToolkitHost> EditWithinLevelEditor = TSharedPtr<IToolkitHost>() ) override;
 	virtual bool CanFilter() override { return false; }
 	virtual uint32 GetCategories() override { return EAssetTypeCategories::MaterialsAndTextures; }
diff --git a/Engine/Source/Developer/AssetTools/Private/AssetTypeActions/AssetTypeActions_Texture2D.cpp b/Engine/Source/Developer/AssetTools/Private/AssetTypeActions/AssetTypeActions_Texture2D.cpp
--- a/Engine/Source/Developer/AssetTools/Private/AssetTypeActions/AssetTypeActions_Texture2D.cpp
+++ b/Engine/Source/Developer/AssetTools/Private/AssetTypeActions/AssetTypeActions_Texture2D.cpp
@@ -1,7 +1,7 @@
 // Copyright 1998-2019 Epic Games, Inc. All Rights Reserved.
 
 #include "AssetTypeActions/AssetTypeActions_Texture2D.h"
-#include "Framework/MultiBox/MultiBoxBuilder.h"
+#include "EditorMenuSubsystem.h"
 #include "Misc/PackageName.h"
 #include "EditorStyleSet.h"
 #include "Factories/SlateBrushAssetFactory.h"
@@ -14,13 +14,14 @@
 
 #define LOCTEXT_NAMESPACE "AssetTypeActions"
 
-void FAssetTypeActions_Texture2D::GetActions( const TArray<UObject*>& InObjects, FMenuBuilder& MenuBuilder )
+void FAssetTypeActions_Texture2D::GetActions(const TArray<UObject*>& InObjects, FEditorMenuSection& Section)
 {
-	FAssetTypeActions_Texture::GetActions(InObjects, MenuBuilder);
+	FAssetTypeActions_Texture::GetActions(InObjects, Section);
 
 	auto Textures = GetTypedWeakObjectPtrs<UTexture2D>(InObjects);
 
-	MenuBuilder.AddMenuEntry(
+	Section.AddMenuEntry(
+		"Texture2D_CreateSlateBrush",
 		LOCTEXT("Texture2D_CreateSlateBrush", "Create Slate Brush"),
 		LOCTEXT("Texture2D_CreateSlateBrushToolTip", "Creates a new slate brush using this texture."),
 		FSlateIcon(FEditorStyle::GetStyleSetName(), "ClassIcon.SlateBrushAsset"),
@@ -30,7 +31,8 @@ void FAssetTypeActions_Texture2D::GetActions( const TArray<UObject*>& InObjects,
 	static const auto AllowVolumeTextureAssetCreationVar = IConsoleManager::Get().FindTConsoleVariableDataInt(TEXT("r.AllowVolumeTextureAssetCreation"));
 	if (InObjects.Num() == 1 && AllowVolumeTextureAssetCreationVar->GetValueOnGameThread() != 0)
 	{
-		MenuBuilder.AddMenuEntry(
+		Section.AddMenuEntry(
+			"Texture2D_CreateVolumeTexture",
 			LOCTEXT("Texture2D_CreateVolumeTexture", "Create Volume Texture"),
 			LOCTEXT("Texture2D_CreateVolumeTextureToolTip", "Creates a new volume texture using this texture."),
 			FSlateIcon(FEditorStyle::GetStyleSetName(), "ClassIcon.Sphere"),
diff --git a/Engine/Source/Developer/AssetTools/Private/AssetTypeActions/AssetTypeActions_Texture2D.h b/Engine/Source/Developer/AssetTools/Private/AssetTypeActions/AssetTypeActions_Texture2D.h
--- a/Engine/Source/Developer/AssetTools/Private/AssetTypeActions/AssetTypeActions_Texture2D.h
+++ b/Engine/Source/Developer/AssetTools/Private/AssetTypeActions/AssetTypeActions_Texture2D.h
@@ -6,7 +6,6 @@
 #include "Engine/Texture2D.h"
 #include "AssetTypeActions/AssetTypeActions_Texture.h"
 
-class FMenuBuilder;
 
 class FAssetTypeActions_Texture2D : public FAssetTypeActions_Texture
 {
@@ -16,7 +15,7 @@ public:
 	virtual FColor GetTypeColor() const override { return FColor(192,64,64); }
 	virtual UClass* GetSupportedClass() const override { return UTexture2D::StaticClass(); }
 	virtual bool HasActions ( const TArray<UObject*>& InObjects ) const override { return true; }
-	virtual void GetActions( const TArray<UObject*>& InObjects, FMenuBuilder& MenuBuilder ) override;
+	virtual void GetActions(const TArray<UObject*>& InObjects, struct FEditorMenuSection& Section) override;
 	virtual bool CanFilter() override { return true; }
 	virtual uint32 GetCategories() override { return FAssetTypeActions_Texture::GetCategories() | EAssetTypeCategories::Basic; }
 
diff --git a/Engine/Source/Developer/AssetTools/Private/AssetTypeActions/AssetTypeActions_TextureRenderTarget.cpp b/Engine/Source/Developer/AssetTools/Private/AssetTypeActions/AssetTypeActions_TextureRenderTarget.cpp
--- a/Engine/Source/Developer/AssetTools/Private/AssetTypeActions/AssetTypeActions_TextureRenderTarget.cpp
+++ b/Engine/Source/Developer/AssetTools/Private/AssetTypeActions/AssetTypeActions_TextureRenderTarget.cpp
@@ -1,7 +1,7 @@
 // Copyright 1998-2019 Epic Games, Inc. All Rights Reserved.
 
 #include "AssetTypeActions/AssetTypeActions_TextureRenderTarget.h"
-#include "Framework/MultiBox/MultiBoxBuilder.h"
+#include "EditorMenuSubsystem.h"
 #include "Engine/Texture2D.h"
 #include "EditorStyleSet.h"
 #include "Engine/TextureCube.h"
@@ -11,13 +11,14 @@
 
 #define LOCTEXT_NAMESPACE "AssetTypeActions"
 
-void FAssetTypeActions_TextureRenderTarget::GetActions( const TArray<UObject*>& InObjects, FMenuBuilder& MenuBuilder )
+void FAssetTypeActions_TextureRenderTarget::GetActions(const TArray<UObject*>& InObjects, FEditorMenuSection& Section)
 {
-	FAssetTypeActions_Texture::GetActions(InObjects, MenuBuilder);
+	FAssetTypeActions_Texture::GetActions(InObjects, Section);
 
 	auto RenderTargets = GetTypedWeakObjectPtrs<UTextureRenderTarget>(InObjects);
 
-	MenuBuilder.AddMenuEntry(
+	Section.AddMenuEntry(
+		"TextureRenderTarget_CreateStatic",
 		LOCTEXT("TextureRenderTarget_CreateStatic", "Create Static Texture"),
 		LOCTEXT("TextureRenderTarget_CreateStaticTooltip", "Creates a static texture from the selected render targets."),
 		FSlateIcon(FEditorStyle::GetStyleSetName(), "ClassIcon.Texture2D"),
diff --git a/Engine/Source/Developer/AssetTools/Private/AssetTypeActions/AssetTypeActions_TextureRenderTarget.h b/Engine/Source/Developer/AssetTools/Private/AssetTypeActions/AssetTypeActions_TextureRenderTarget.h
--- a/Engine/Source/Developer/AssetTools/Private/AssetTypeActions/AssetTypeActions_TextureRenderTarget.h
+++ b/Engine/Source/Developer/AssetTools/Private/AssetTypeActions/AssetTypeActions_TextureRenderTarget.h
@@ -6,7 +6,6 @@
 #include "AssetTypeActions/AssetTypeActions_Texture.h"
 #include "Engine/TextureRenderTarget.h"
 
-class FMenuBuilder;
 
 class FAssetTypeActions_TextureRenderTarget : public FAssetTypeActions_Texture
 {
@@ -16,7 +15,7 @@ public:
 	virtual FColor GetTypeColor() const override { return FColor(128,64,64); }
 	virtual UClass* GetSupportedClass() const override { return UTextureRenderTarget::StaticClass(); }
 	virtual bool HasActions ( const TArray<UObject*>& InObjects ) const override { return true; }
-	virtual void GetActions( const TArray<UObject*>& InObjects, FMenuBuilder& MenuBuilder ) override;
+	virtual void GetActions(const TArray<UObject*>& InObjects, struct FEditorMenuSection& Section) override;
 	virtual bool IsImportedAsset() const override { return false; }
 
 private:
diff --git a/Engine/Source/Developer/AssetTools/Public/AssetTypeActions/AssetTypeActions_Blueprint.h b/Engine/Source/Developer/AssetTools/Public/AssetTypeActions/AssetTypeActions_Blueprint.h
--- a/Engine/Source/Developer/AssetTools/Public/AssetTypeActions/AssetTypeActions_Blueprint.h
+++ b/Engine/Source/Developer/AssetTools/Public/AssetTypeActions/AssetTypeActions_Blueprint.h
@@ -8,7 +8,6 @@
 #include "AssetTypeActions/AssetTypeActions_ClassTypeBase.h"
 
 struct FAssetData;
-class FMenuBuilder;
 class IClassTypeActions;
 class UFactory;
 
@@ -20,7 +19,7 @@ public:
 	virtual FColor GetTypeColor() const override { return FColor( 63, 126, 255 ); }
 	virtual UClass* GetSupportedClass() const override { return UBlueprint::StaticClass(); }
 	virtual bool HasActions ( const TArray<UObject*>& InObjects ) const override { return true; }
-	virtual void GetActions( const TArray<UObject*>& InObjects, FMenuBuilder& MenuBuilder ) override;
+	virtual void GetActions(const TArray<UObject*>& InObjects, FEditorMenuSection& Section) override;
 	virtual void OpenAssetEditor(const TArray<UObject*>& InObjects, TSharedPtr<class IToolkitHost> EditWithinLevelEditor = TSharedPtr<IToolkitHost>()) override;
 	virtual bool CanMerge() const override;
 	virtual void Merge(UObject* InObject) override;
diff --git a/Engine/Source/Editor/AudioEditor/AudioEditor.Build.cs b/Engine/Source/Editor/AudioEditor/AudioEditor.Build.cs
--- a/Engine/Source/Editor/AudioEditor/AudioEditor.Build.cs
+++ b/Engine/Source/Editor/AudioEditor/AudioEditor.Build.cs
@@ -15,7 +15,8 @@ public class AudioEditor : ModuleRules
 
 		PrivateDependencyModuleNames.AddRange(
 			new string[] {
-				"AudioMixer"
+				"AudioMixer",
+				"EditorMenus",
 			});
 
 		PublicDependencyModuleNames.AddRange(
diff --git a/Engine/Source/Editor/AudioEditor/Private/AssetTypeActions/AssetTypeActions_DialogueWave.cpp b/Engine/Source/Editor/AudioEditor/Private/AssetTypeActions/AssetTypeActions_DialogueWave.cpp
--- a/Engine/Source/Editor/AudioEditor/Private/AssetTypeActions/AssetTypeActions_DialogueWave.cpp
+++ b/Engine/Source/Editor/AudioEditor/Private/AssetTypeActions/AssetTypeActions_DialogueWave.cpp
@@ -2,7 +2,7 @@
 
 #include "AssetTypeActions/AssetTypeActions_DialogueWave.h"
 #include "Misc/PackageName.h"
-#include "Framework/MultiBox/MultiBoxBuilder.h"
+#include "EditorMenuSubsystem.h"
 #include "Editor.h"
 #include "EditorStyleSet.h"
 #include "Factories/SoundCueFactoryNew.h"
@@ -12,11 +12,12 @@
 
 #define LOCTEXT_NAMESPACE "AssetTypeActions"
 
-void FAssetTypeActions_DialogueWave::GetActions(const TArray<UObject*>& InObjects, FMenuBuilder& MenuBuilder)
+void FAssetTypeActions_DialogueWave::GetActions(const TArray<UObject*>& InObjects, FEditorMenuSection& Section)
 {
 	auto DialogueWaves = GetTypedWeakObjectPtrs<UDialogueWave>(InObjects);
 
-	MenuBuilder.AddMenuEntry(
+	Section.AddMenuEntry(
+		"Sound_PlaySound",
 		LOCTEXT("Sound_PlaySound", "Play"),
 		LOCTEXT("Sound_PlaySoundTooltip", "Plays the selected sound."),
 		FSlateIcon(FEditorStyle::GetStyleSetName(), "MediaAsset.AssetActions.Play.Small"),
@@ -26,7 +27,8 @@ void FAssetTypeActions_DialogueWave::GetActions(const TArray<UObject*>& InObject
 		)
 		);
 
-	MenuBuilder.AddMenuEntry(
+	Section.AddMenuEntry(
+		"Sound_StopSound",
 		LOCTEXT("Sound_StopSound", "Stop"),
 		LOCTEXT("Sound_StopSoundTooltip", "Stops the selected sounds."),
 		FSlateIcon(FEditorStyle::GetStyleSetName(), "MediaAsset.AssetActions.Stop.Small"),
@@ -36,7 +38,8 @@ void FAssetTypeActions_DialogueWave::GetActions(const TArray<UObject*>& InObject
 		)
 		);
 
-	MenuBuilder.AddMenuEntry(
+	Section.AddMenuEntry(
+		"DialogueWave_CreateCue",
 		LOCTEXT("DialogueWave_CreateCue", "Create Cue"),
 		LOCTEXT("DialogueWave_CreateCueTooltip", "Creates a sound cue using this dialogue wave."),
 		FSlateIcon(FEditorStyle::GetStyleSetName(), "ClassIcon.SoundCue"),
diff --git a/Engine/Source/Editor/AudioEditor/Private/AssetTypeActions/AssetTypeActions_DialogueWave.h b/Engine/Source/Editor/AudioEditor/Private/AssetTypeActions/AssetTypeActions_DialogueWave.h
--- a/Engine/Source/Editor/AudioEditor/Private/AssetTypeActions/AssetTypeActions_DialogueWave.h
+++ b/Engine/Source/Editor/AudioEditor/Private/AssetTypeActions/AssetTypeActions_DialogueWave.h
@@ -7,7 +7,6 @@
 #include "AssetTypeActions_Base.h"
 #include "Sound/DialogueWave.h"
 
-class FMenuBuilder;
 
 class FAssetTypeActions_DialogueWave : public FAssetTypeActions_Base
 {
@@ -19,7 +18,7 @@ public:
 	virtual uint32 GetCategories() override { return EAssetTypeCategories::Sounds; }
 	virtual bool CanFilter() override { return true; }
 	virtual bool HasActions(const TArray<UObject*>& InObjects) const override { return true; }
-	virtual void GetActions(const TArray<UObject*>& InObjects, FMenuBuilder& MenuBuilder) override;
+	virtual void GetActions(const TArray<UObject*>& InObjects, struct FEditorMenuSection& Section) override;
 	virtual void OpenAssetEditor( const TArray<UObject*>& InObjects, TSharedPtr<class IToolkitHost> EditWithinLevelEditor = TSharedPtr<IToolkitHost>() ) override;
 
 private:
diff --git a/Engine/Source/Editor/AudioEditor/Private/AssetTypeActions/AssetTypeActions_SoundBase.cpp b/Engine/Source/Editor/AudioEditor/Private/AssetTypeActions/AssetTypeActions_SoundBase.cpp
--- a/Engine/Source/Editor/AudioEditor/Private/AssetTypeActions/AssetTypeActions_SoundBase.cpp
+++ b/Engine/Source/Editor/AudioEditor/Private/AssetTypeActions/AssetTypeActions_SoundBase.cpp
@@ -2,7 +2,7 @@
 
 #include "AssetTypeActions/AssetTypeActions_SoundBase.h"
 #include "Sound/SoundBase.h"
-#include "Framework/MultiBox/MultiBoxBuilder.h"
+#include "EditorMenuSubsystem.h"
 #include "Editor.h"
 #include "EditorStyleSet.h"
 #include "Components/AudioComponent.h"
@@ -17,11 +17,12 @@ UClass* FAssetTypeActions_SoundBase::GetSupportedClass() const
 	return USoundBase::StaticClass();
 }
 
-void FAssetTypeActions_SoundBase::GetActions( const TArray<UObject*>& InObjects, FMenuBuilder& MenuBuilder )
+void FAssetTypeActions_SoundBase::GetActions(const TArray<UObject*>& InObjects, FEditorMenuSection& Section)
 {
 	auto Sounds = GetTypedWeakObjectPtrs<USoundBase>(InObjects);
 
-	MenuBuilder.AddMenuEntry(
+	Section.AddMenuEntry(
+		"Sound_PlaySound",
 		LOCTEXT("Sound_PlaySound", "Play"),
 		LOCTEXT("Sound_PlaySoundTooltip", "Plays the selected sound."),
 		FSlateIcon(FEditorStyle::GetStyleSetName(), "MediaAsset.AssetActions.Play.Small"),
@@ -31,7 +32,8 @@ void FAssetTypeActions_SoundBase::GetActions( const TArray<UObject*>& InObjects,
 			)
 		);
 
-	MenuBuilder.AddMenuEntry(
+	Section.AddMenuEntry(
+		"Sound_StopSound",
 		LOCTEXT("Sound_StopSound", "Stop"),
 		LOCTEXT("Sound_StopSoundTooltip", "Stops the selected sounds."),
 		FSlateIcon(FEditorStyle::GetStyleSetName(), "MediaAsset.AssetActions.Stop.Small"),
diff --git a/Engine/Source/Editor/AudioEditor/Private/AssetTypeActions/AssetTypeActions_SoundCue.cpp b/Engine/Source/Editor/AudioEditor/Private/AssetTypeActions/AssetTypeActions_SoundCue.cpp
--- a/Engine/Source/Editor/AudioEditor/Private/AssetTypeActions/AssetTypeActions_SoundCue.cpp
+++ b/Engine/Source/Editor/AudioEditor/Private/AssetTypeActions/AssetTypeActions_SoundCue.cpp
@@ -4,7 +4,7 @@
 #include "Sound/SoundAttenuation.h"
 #include "Misc/PackageName.h"
 #include "Sound/SoundCue.h"
-#include "Framework/MultiBox/MultiBoxBuilder.h"
+#include "EditorMenuSubsystem.h"
 #include "Factories/SoundAttenuationFactory.h"
 #include "ContentBrowserModule.h"
 #include "IContentBrowserSingleton.h"
@@ -17,12 +17,13 @@ UClass* FAssetTypeActions_SoundCue::GetSupportedClass() const
 	return USoundCue::StaticClass();
 }
 
-void FAssetTypeActions_SoundCue::GetActions( const TArray<UObject*>& InObjects, FMenuBuilder& MenuBuilder )
+void FAssetTypeActions_SoundCue::GetActions(const TArray<UObject*>& InObjects, FEditorMenuSection& Section)
 {
 	auto SoundCues = GetTypedWeakObjectPtrs<USoundCue>(InObjects);
-	FAssetTypeActions_SoundBase::GetActions(InObjects, MenuBuilder);
+	FAssetTypeActions_SoundBase::GetActions(InObjects, Section);
 
-	MenuBuilder.AddMenuEntry(
+	Section.AddMenuEntry(
+		"SoundCue_ConsolidateAttenuation",
 		LOCTEXT("SoundCue_ConsolidateAttenuation", "Consolidate Attenuation"),
 		LOCTEXT("SoundCue_ConsolidateAttenuationTooltip", "Creates shared attenuation packages for sound cues with identical override attenuation settings."),
 		FSlateIcon(),
diff --git a/Engine/Source/Editor/AudioEditor/Private/AssetTypeActions/AssetTypeActions_SoundCue.h b/Engine/Source/Editor/AudioEditor/Private/AssetTypeActions/AssetTypeActions_SoundCue.h
--- a/Engine/Source/Editor/AudioEditor/Private/AssetTypeActions/AssetTypeActions_SoundCue.h
+++ b/Engine/Source/Editor/AudioEditor/Private/AssetTypeActions/AssetTypeActions_SoundCue.h
@@ -6,7 +6,6 @@
 #include "Toolkits/IToolkitHost.h"
 #include "AssetTypeActions/AssetTypeActions_SoundBase.h"
 
-class FMenuBuilder;
 class USoundCue;
 
 class FAssetTypeActions_SoundCue : public FAssetTypeActions_SoundBase
@@ -16,7 +15,7 @@ public:
 	virtual FText GetName() const override { return NSLOCTEXT("AssetTypeActions", "AssetTypeActions_SoundCue", "Sound Cue"); }
 	virtual FColor GetTypeColor() const override { return FColor(0, 175, 255); }
 	virtual UClass* GetSupportedClass() const override;
-	virtual void GetActions( const TArray<UObject*>& InObjects, FMenuBuilder& MenuBuilder ) override;
+	virtual void GetActions(const TArray<UObject*>& InObjects, struct FEditorMenuSection& Section) override;
 	virtual void OpenAssetEditor( const TArray<UObject*>& InObjects, TSharedPtr<class IToolkitHost> EditWithinLevelEditor = TSharedPtr<IToolkitHost>() ) override;
 	virtual bool CanFilter() override { return true; }
 
diff --git a/Engine/Source/Editor/AudioEditor/Private/AssetTypeActions/AssetTypeActions_SoundCueTemplate.cpp b/Engine/Source/Editor/AudioEditor/Private/AssetTypeActions/AssetTypeActions_SoundCueTemplate.cpp
--- a/Engine/Source/Editor/AudioEditor/Private/AssetTypeActions/AssetTypeActions_SoundCueTemplate.cpp
+++ b/Engine/Source/Editor/AudioEditor/Private/AssetTypeActions/AssetTypeActions_SoundCueTemplate.cpp
@@ -5,7 +5,7 @@
 #include "EditorStyleSet.h"
 #include "Factories/SoundCueFactoryNew.h"
 #include "Factories/SoundCueTemplateFactory.h"
-#include "Framework/MultiBox/MultiBoxBuilder.h"
+#include "EditorMenuSubsystem.h"
 #include "IContentBrowserSingleton.h"
 #include "Sound/SoundCue.h"
 #include "Sound/SoundCueTemplate.h"
@@ -21,11 +21,12 @@ UClass* FAssetTypeActions_SoundCueTemplate::GetSupportedClass() const
 	return USoundCueTemplate::StaticClass();
 }
 
-void FAssetTypeActions_SoundCueTemplate::GetActions(const TArray<UObject*>& InObjects, FMenuBuilder& MenuBuilder)
+void FAssetTypeActions_SoundCueTemplate::GetActions(const TArray<UObject*>& InObjects, FEditorMenuSection& Section)
 {
 	TArray<TWeakObjectPtr<USoundCueTemplate>> Cues = GetTypedWeakObjectPtrs<USoundCueTemplate>(InObjects);
 
-	MenuBuilder.AddMenuEntry(
+	Section.AddMenuEntry(
+		"SoundCueTemplate_CopyToSoundCue",
 		LOCTEXT("SoundCueTemplate_CopyToSoundCue", "Copy To Sound Cue"),
 		LOCTEXT("SoundCueTemplate_CopyToSoundCueTooltip", "Exports a Sound Cue Template to a Sound Cue."),
 		FSlateIcon(FEditorStyle::GetStyleSetName(), "ClassIcon.SoundCue"),
diff --git a/Engine/Source/Editor/AudioEditor/Private/AssetTypeActions/AssetTypeActions_SoundCueTemplate.h b/Engine/Source/Editor/AudioEditor/Private/AssetTypeActions/AssetTypeActions_SoundCueTemplate.h
--- a/Engine/Source/Editor/AudioEditor/Private/AssetTypeActions/AssetTypeActions_SoundCueTemplate.h
+++ b/Engine/Source/Editor/AudioEditor/Private/AssetTypeActions/AssetTypeActions_SoundCueTemplate.h
@@ -5,7 +5,6 @@
 #include "Toolkits/IToolkitHost.h"
 #include "AssetTypeActions/AssetTypeActions_SoundBase.h"
 
-class FMenuBuilder;
 class USoundCueTemplate;
 
 class FAssetTypeActions_SoundCueTemplate : public FAssetTypeActions_SoundBase
@@ -15,7 +14,7 @@ public:
 	virtual FText GetName() const override { return NSLOCTEXT("AssetTypeActions", "AssetTypeActions_SoundCueTemplate", "Sound Cue Template"); }
 	virtual FColor GetTypeColor() const override { return FColor(255, 255, 255); }
 	virtual UClass* GetSupportedClass() const override;
-	virtual void GetActions(const TArray<UObject*>& InObjects, FMenuBuilder& MenuBuilder) override;
+	virtual void GetActions(const TArray<UObject*>& InObjects, struct FEditorMenuSection& Section) override;
 	virtual bool CanFilter() override { return true; }
 
 protected:
diff --git a/Engine/Source/Editor/AudioEditor/Private/AssetTypeActions/AssetTypeActions_SoundWave.cpp b/Engine/Source/Editor/AudioEditor/Private/AssetTypeActions/AssetTypeActions_SoundWave.cpp
--- a/Engine/Source/Editor/AudioEditor/Private/AssetTypeActions/AssetTypeActions_SoundWave.cpp
+++ b/Engine/Source/Editor/AudioEditor/Private/AssetTypeActions/AssetTypeActions_SoundWave.cpp
@@ -5,7 +5,7 @@
 #include "Misc/PackageName.h"
 #include "AssetData.h"
 #include "Sound/SoundWave.h"
-#include "Framework/MultiBox/MultiBoxBuilder.h"
+#include "EditorMenuSubsystem.h"
 #include "EditorStyleSet.h"
 #include "Factories/SoundCueFactoryNew.h"
 #include "EditorFramework/AssetImportData.h"
@@ -27,13 +27,14 @@ UClass* FAssetTypeActions_SoundWave::GetSupportedClass() const
 	return USoundWave::StaticClass();
 }
 
-void FAssetTypeActions_SoundWave::GetActions( const TArray<UObject*>& InObjects, FMenuBuilder& MenuBuilder )
+void FAssetTypeActions_SoundWave::GetActions(const TArray<UObject*>& InObjects, FEditorMenuSection& Section)
 {
-	FAssetTypeActions_SoundBase::GetActions(InObjects, MenuBuilder);
+	FAssetTypeActions_SoundBase::GetActions(InObjects, Section);
 
 	TArray<TWeakObjectPtr<USoundWave>> SoundNodes = GetTypedWeakObjectPtrs<USoundWave>(InObjects);
 
-	MenuBuilder.AddMenuEntry(
+	Section.AddMenuEntry(
+		"SoundWave_CreateCue",
 		LOCTEXT("SoundWave_CreateCue", "Create Cue"),
 		LOCTEXT("SoundWave_CreateCueTooltip", "Creates a sound cue using this sound wave."),
 		FSlateIcon(FEditorStyle::GetStyleSetName(), "ClassIcon.SoundCue"),
@@ -43,21 +44,12 @@ void FAssetTypeActions_SoundWave::GetActions( const TArray<UObject*>& InObjects,
 			)
 		);
 
-	MenuBuilder.AddSubMenu(
+	Section.AddEntry(FEditorMenuEntry::InitSubMenu(
+		NAME_None,
+		"SoundWave_CreateDialogue",
 		LOCTEXT("SoundWave_CreateDialogue", "Create Dialogue"),
 		LOCTEXT("SoundWave_CreateDialogueTooltip", "Creates a dialogue wave using this sound wave."),
-		FNewMenuDelegate::CreateSP(this, &FAssetTypeActions_SoundWave::FillVoiceMenu, SoundNodes));
-
-	// Check for any sound wave asset action extensions
-	IAudioEditorModule* AudioEditorModule = &FModuleManager::LoadModuleChecked<IAudioEditorModule>("AudioEditor");
-	TArray<TSharedPtr<ISoundWaveAssetActionExtensions>> Extensions;
-	AudioEditorModule->GetSoundWaveActionExtenders(Extensions);
-	
-	for (TSharedPtr<ISoundWaveAssetActionExtensions> Extension : Extensions)
-	{
-		Extension->GetExtendedActions(SoundNodes, MenuBuilder);
-	}
-
+		FNewMenuDelegate::CreateSP(this, &FAssetTypeActions_SoundWave::FillVoiceMenu, SoundNodes)));
 }
 
 void FAssetTypeActions_SoundWave::GetResolvedSourceFilePaths(const TArray<UObject*>& TypeAssets, TArray<FString>& OutSourceFilePaths) const
diff --git a/Engine/Source/Editor/AudioEditor/Private/AssetTypeActions/AssetTypeActions_SoundWave.h b/Engine/Source/Editor/AudioEditor/Private/AssetTypeActions/AssetTypeActions_SoundWave.h
--- a/Engine/Source/Editor/AudioEditor/Private/AssetTypeActions/AssetTypeActions_SoundWave.h
+++ b/Engine/Source/Editor/AudioEditor/Private/AssetTypeActions/AssetTypeActions_SoundWave.h
@@ -6,7 +6,6 @@
 #include "Toolkits/IToolkitHost.h"
 #include "AssetTypeActions/AssetTypeActions_SoundBase.h"
 
-class FMenuBuilder;
 class USoundWave;
 
 class FAssetTypeActions_SoundWave : public FAssetTypeActions_SoundBase
@@ -16,7 +15,7 @@ public:
 	virtual FText GetName() const override { return NSLOCTEXT("AssetTypeActions", "AssetTypeActions_SoundWave", "Sound Wave"); }
 	virtual FColor GetTypeColor() const override { return FColor(97, 85, 212); }
 	virtual UClass* GetSupportedClass() const override;
-	virtual void GetActions( const TArray<UObject*>& InObjects, FMenuBuilder& MenuBuilder ) override;
+	virtual void GetActions(const TArray<UObject*>& InObjects, struct FEditorMenuSection& Section) override;
 	virtual void OpenAssetEditor( const TArray<UObject*>& InObjects, TSharedPtr<class IToolkitHost> EditWithinLevelEditor = TSharedPtr<IToolkitHost>() ) override;
 	virtual bool CanFilter() override { return true; }
 	virtual bool IsImportedAsset() const override { return true; }
diff --git a/Engine/Source/Editor/AudioEditor/Private/AudioEditorModule.cpp b/Engine/Source/Editor/AudioEditor/Private/AudioEditorModule.cpp
--- a/Engine/Source/Editor/AudioEditor/Private/AudioEditorModule.cpp
+++ b/Engine/Source/Editor/AudioEditor/Private/AudioEditorModule.cpp
@@ -148,19 +148,6 @@ public:
 		}
 	}
 
-	virtual void AddSoundWaveActionExtender(TSharedPtr<ISoundWaveAssetActionExtensions> InSoundWaveAssetActionExtender) override
-	{
-		if (InSoundWaveAssetActionExtender.IsValid())
-		{
-			SoundWaveAssetActionExtensions.AddUnique(InSoundWaveAssetActionExtender);
-		}
-	}
-
-	virtual void GetSoundWaveActionExtenders(TArray<TSharedPtr<ISoundWaveAssetActionExtensions>>& OutSoundwaveActionExtensions) override
-	{
-		OutSoundwaveActionExtensions = SoundWaveAssetActionExtensions;
-	}
-
 	virtual void RegisterEffectPresetAssetActions() override
 	{
 		// Only register asset actions for the case where audio mixer data is enabled
@@ -337,7 +324,6 @@ private:
 	FExtensibilityManagers SoundCueExtensibility;
 	FExtensibilityManagers SoundClassExtensibility;
 	FExtensibilityManagers SoundSubmixExtensibility;
-	TArray<TSharedPtr<ISoundWaveAssetActionExtensions>> SoundWaveAssetActionExtensions;
 	TSet<USoundEffectPreset*> RegisteredActions;
 	TSharedPtr<FGraphPanelPinConnectionFactory> SoundCueGraphConnectionFactory;
 	TSharedPtr<FSlateStyleSet> AudioStyleSet;
diff --git a/Engine/Source/Editor/AudioEditor/Public/AssetTypeActions/AssetTypeActions_SoundBase.h b/Engine/Source/Editor/AudioEditor/Public/AssetTypeActions/AssetTypeActions_SoundBase.h
--- a/Engine/Source/Editor/AudioEditor/Public/AssetTypeActions/AssetTypeActions_SoundBase.h
+++ b/Engine/Source/Editor/AudioEditor/Public/AssetTypeActions/AssetTypeActions_SoundBase.h
@@ -6,7 +6,6 @@
 #include "AssetTypeActions_Base.h"
 
 struct FAssetData;
-class FMenuBuilder;
 class USoundBase;
 
 class AUDIOEDITOR_API FAssetTypeActions_SoundBase : public FAssetTypeActions_Base
@@ -17,7 +16,7 @@ public:
 	virtual FColor GetTypeColor() const override { return FColor(97, 85, 212); }
 	virtual UClass* GetSupportedClass() const override;
 	virtual bool HasActions ( const TArray<UObject*>& InObjects ) const override { return true; }
-	virtual void GetActions( const TArray<UObject*>& InObjects, FMenuBuilder& MenuBuilder ) override;
+	virtual void GetActions(const TArray<UObject*>& InObjects, struct FEditorMenuSection& Section) override;
 	virtual void AssetsActivated( const TArray<UObject*>& InObjects, EAssetTypeActivationMethod::Type ActivationType ) override;
 	virtual uint32 GetCategories() override { return EAssetTypeCategories::Sounds; }
 	virtual bool CanFilter() override { return false; }
diff --git a/Engine/Source/Editor/AudioEditor/Public/AudioEditorModule.h b/Engine/Source/Editor/AudioEditor/Public/AudioEditorModule.h
--- a/Engine/Source/Editor/AudioEditor/Public/AudioEditorModule.h
+++ b/Engine/Source/Editor/AudioEditor/Public/AudioEditorModule.h
@@ -21,13 +21,6 @@ AUDIOEDITOR_API DECLARE_LOG_CATEGORY_EXTERN(LogAudioEditor, Log, All);
 
 extern const FName AudioEditorAppIdentifier;
 
-/** Interface which can be implemented in plugin to extend sound wave asset actions. */
-class ISoundWaveAssetActionExtensions
-{
-public:
-	virtual void GetExtendedActions(const TArray<TWeakObjectPtr<USoundWave>>& InObjects, FMenuBuilder& MenuBuilder) = 0;
-};
-
 /** Sound class editor module interface */
 class IAudioEditorModule :	public IModuleInterface
 {
@@ -41,12 +34,6 @@ public:
 	/** Registers effect preset asset actions. */
 	virtual void RegisterEffectPresetAssetActions() {}
 
-	/** Adds an a sound wave asset action extender class to allow extending asset wave actions to plugins. */
-	virtual void AddSoundWaveActionExtender(TSharedPtr<ISoundWaveAssetActionExtensions> InSoundWaveAssetActionExtender) = 0;
-
-	/** Returns all added sound wave action extenders currently added. */
-	virtual void GetSoundWaveActionExtenders(TArray<TSharedPtr<ISoundWaveAssetActionExtensions>>& OutSoundwaveActionExtensions) = 0;
-
 	/** Creates a new sound class editor for a sound class object. */
 	virtual TSharedRef<FAssetEditorToolkit> CreateSoundClassEditor( const EToolkitMode::Type Mode, const TSharedPtr< IToolkitHost >& InitToolkitHost, USoundClass* InSoundClass ) = 0;
 
diff --git a/Engine/Source/Editor/Blutility/Blutility.Build.cs b/Engine/Source/Editor/Blutility/Blutility.Build.cs
--- a/Engine/Source/Editor/Blutility/Blutility.Build.cs
+++ b/Engine/Source/Editor/Blutility/Blutility.Build.cs
@@ -36,8 +36,9 @@ public class Blutility : ModuleRules
                 "BlueprintGraph",
 				"UMG",
                 "UMGEditor",
-                "KismetCompiler"
-            }
+                "KismetCompiler",
+				"EditorMenus",
+			}
 			);
 	}
 }
diff --git a/Engine/Source/Editor/Blutility/Private/AssetTypeActions_EditorUtilityBlueprint.cpp b/Engine/Source/Editor/Blutility/Private/AssetTypeActions_EditorUtilityBlueprint.cpp
--- a/Engine/Source/Editor/Blutility/Private/AssetTypeActions_EditorUtilityBlueprint.cpp
+++ b/Engine/Source/Editor/Blutility/Private/AssetTypeActions_EditorUtilityBlueprint.cpp
@@ -1,7 +1,7 @@
 // Copyright 1998-2019 Epic Games, Inc. All Rights Reserved.
 
 #include "AssetTypeActions_EditorUtilityBlueprint.h"
-#include "Framework/MultiBox/MultiBoxBuilder.h"
+#include "EditorMenuSubsystem.h"
 #include "Misc/PackageName.h"
 #include "Misc/MessageDialog.h"
 #include "EditorUtilityBlueprintFactory.h"
@@ -38,11 +38,12 @@ bool FAssetTypeActions_EditorUtilityBlueprint::HasActions(const TArray<UObject*>
 	return true;
 }
 
-void FAssetTypeActions_EditorUtilityBlueprint::GetActions(const TArray<UObject*>& InObjects, FMenuBuilder& MenuBuilder)
+void FAssetTypeActions_EditorUtilityBlueprint::GetActions(const TArray<UObject*>& InObjects, FEditorMenuSection& Section)
 {
 	auto Blueprints = GetTypedWeakObjectPtrs<UEditorUtilityBlueprint>(InObjects);
 
-	MenuBuilder.AddMenuEntry(
+	Section.AddMenuEntry(
+		"EditorUtility_Run",
 		LOCTEXT("EditorUtility_Run", "Run Editor Utility Blueprint"),
 		LOCTEXT("EditorUtility_RunTooltip", "Runs this Editor Utility Blueprint."),
 		FSlateIcon(),
diff --git a/Engine/Source/Editor/Blutility/Private/AssetTypeActions_EditorUtilityBlueprint.h b/Engine/Source/Editor/Blutility/Private/AssetTypeActions_EditorUtilityBlueprint.h
--- a/Engine/Source/Editor/Blutility/Private/AssetTypeActions_EditorUtilityBlueprint.h
+++ b/Engine/Source/Editor/Blutility/Private/AssetTypeActions_EditorUtilityBlueprint.h
@@ -7,7 +7,6 @@
 #include "Toolkits/IToolkitHost.h"
 #include "AssetTypeActions/AssetTypeActions_Blueprint.h"
 
-class FMenuBuilder;
 
 class FAssetTypeActions_EditorUtilityBlueprint : public FAssetTypeActions_Blueprint
 {
@@ -17,7 +16,7 @@ public:
 	virtual FColor GetTypeColor() const override;
 	virtual UClass* GetSupportedClass() const override;
 	virtual bool HasActions(const TArray<UObject*>& InObjects) const override;
-	virtual void GetActions(const TArray<UObject*>& InObjects, FMenuBuilder& MenuBuilder) override;
+	virtual void GetActions(const TArray<UObject*>& InObjects, struct FEditorMenuSection& Section) override;
 	virtual uint32 GetCategories() override;
 	virtual bool CanLocalize() const override { return false; }
 	// End of IAssetTypeActions interface
diff --git a/Engine/Source/Editor/Blutility/Private/AssetTypeActions_EditorUtilityWidgetBlueprint.cpp b/Engine/Source/Editor/Blutility/Private/AssetTypeActions_EditorUtilityWidgetBlueprint.cpp
--- a/Engine/Source/Editor/Blutility/Private/AssetTypeActions_EditorUtilityWidgetBlueprint.cpp
+++ b/Engine/Source/Editor/Blutility/Private/AssetTypeActions_EditorUtilityWidgetBlueprint.cpp
@@ -1,7 +1,7 @@
 // Copyright 1998-2019 Epic Games, Inc. All Rights Reserved.
 
 #include "AssetTypeActions_EditorUtilityWidgetBlueprint.h"
-#include "Framework/MultiBox/MultiBoxBuilder.h"
+#include "EditorMenuSubsystem.h"
 #include "Misc/PackageName.h"
 #include "Misc/MessageDialog.h"
 #include "EditorUtilityBlueprintFactory.h"
@@ -45,11 +45,12 @@ bool FAssetTypeActions_EditorUtilityWidgetBlueprint::HasActions(const TArray<UOb
 	return true;
 }
 
-void FAssetTypeActions_EditorUtilityWidgetBlueprint::GetActions(const TArray<UObject*>& InObjects, FMenuBuilder& MenuBuilder)
+void FAssetTypeActions_EditorUtilityWidgetBlueprint::GetActions(const TArray<UObject*>& InObjects, FEditorMenuSection& Section)
 {
 	auto Blueprints = GetTypedWeakObjectPtrs<UWidgetBlueprint>(InObjects);
 
-	MenuBuilder.AddMenuEntry(
+	Section.AddMenuEntry(
+		"EditorUtilityWidget_Edit",
 		LOCTEXT("EditorUtilityWidget_Edit", "Run Editor Utility Widget"),
 		LOCTEXT("EditorUtilityWidget_EditTooltip", "Opens the tab built by this Editor Utility Widget Blueprint."),
 		FSlateIcon(),
diff --git a/Engine/Source/Editor/Blutility/Private/AssetTypeActions_EditorUtilityWidgetBlueprint.h b/Engine/Source/Editor/Blutility/Private/AssetTypeActions_EditorUtilityWidgetBlueprint.h
--- a/Engine/Source/Editor/Blutility/Private/AssetTypeActions_EditorUtilityWidgetBlueprint.h
+++ b/Engine/Source/Editor/Blutility/Private/AssetTypeActions_EditorUtilityWidgetBlueprint.h
@@ -7,7 +7,6 @@
 #include "Toolkits/IToolkitHost.h"
 #include "AssetTypeActions_Base.h"
 
-class FMenuBuilder;
 
 class FAssetTypeActions_EditorUtilityWidgetBlueprint : public FAssetTypeActions_Base
 {
@@ -17,7 +16,7 @@ public:
 	virtual FColor GetTypeColor() const override;
 	virtual UClass* GetSupportedClass() const override;
 	virtual bool HasActions(const TArray<UObject*>& InObjects) const override;
-	virtual void GetActions(const TArray<UObject*>& InObjects, FMenuBuilder& MenuBuilder) override;
+	virtual void GetActions(const TArray<UObject*>& InObjects, struct FEditorMenuSection& Section) override;
 	virtual void OpenAssetEditor(const TArray<UObject*>& InObjects, TSharedPtr<class IToolkitHost> EditWithinLevelEditor = TSharedPtr<IToolkitHost>()) override;
 	virtual uint32 GetCategories() override;
 	virtual void PerformAssetDiff(UObject* Asset1, UObject* Asset2, const struct FRevisionInfo& OldRevision, const struct FRevisionInfo& NewRevision) const override;
diff --git a/Engine/Source/Editor/EditorMenus/Public/EditorMenuSubsystem.h b/Engine/Source/Editor/EditorMenus/Public/EditorMenuSubsystem.h
--- a/Engine/Source/Editor/EditorMenus/Public/EditorMenuSubsystem.h
+++ b/Engine/Source/Editor/EditorMenus/Public/EditorMenuSubsystem.h
@@ -125,6 +125,27 @@ public:
 	 */
 	static bool IsRunningEditorUI();
 
+	static void RegisterStartupCallback(const FSimpleMulticastDelegate::FDelegate& Delegate)
+	{
+		if (IsRunningEditorUI())
+		{
+			if (UEditorMenuSubsystem::TryGet())
+			{
+				Delegate.Execute();
+			}
+			else
+			{
+				// Wait until editor subsystems have been constructed
+				FCoreDelegates::OnPostEngineInit.Add(Delegate);
+			}
+		}
+	}
+
+	static void UnRegisterStartupCallback(const void* UserPointer)
+	{
+		FCoreDelegates::OnPostEngineInit.RemoveAll(UserPointer);
+	}
+
 	// UEditorSubsystem implementation
 	virtual void Initialize(FSubsystemCollectionBase& Collection);
 	virtual void Deinitialize();
