TimeSynthComponent: prevent access to null audio device when playing with -nosound

diff --git a/Engine/Plugins/Runtime/TimeSynth/Source/TimeSynth/Private/TimeSynthComponent.cpp b/Engine/Plugins/Runtime/TimeSynth/Source/TimeSynth/Private/TimeSynthComponent.cpp
--- a/Engine/Plugins/Runtime/TimeSynth/Source/TimeSynth/Private/TimeSynthComponent.cpp
+++ b/Engine/Plugins/Runtime/TimeSynth/Source/TimeSynth/Private/TimeSynthComponent.cpp
@@ -597,6 +597,17 @@ FTimeSynthClipHandle UTimeSynthComponent::PlayClip(UTimeSynthClip* InClip, UTime
 
 	// Get the distance to nearest listener using this transform
 	const FAudioDevice* OwningAudioDevice = GetAudioDevice();
+
+	// Validate audio device since it might not be available (i.e. -nosound)
+	if (OwningAudioDevice == nullptr)
+	{
+		static bool bShouldWarn = true;
+		UE_CLOG(bShouldWarn, LogTimeSynth, Warning, TEXT("Failed to play clip: no audio device. Running -nosound?"));
+		bShouldWarn = false;
+
+		return FTimeSynthClipHandle();
+	}
+
 	const float DistanceToListener = OwningAudioDevice->GetDistanceToNearestListener(ThisComponentTransform.GetTranslation());
 
 	TArray<FTimeSynthClipSound> ValidSounds;
