changing default value of LastVolumeDb to be different from CurrentVolumeDb on the TimeSynthComponent to make sure volume gets set on first update.

diff --git a/Engine/Plugins/Runtime/TimeSynth/Source/TimeSynth/Classes/TimeSynthComponent.h b/Engine/Plugins/Runtime/TimeSynth/Source/TimeSynth/Classes/TimeSynthComponent.h
--- a/Engine/Plugins/Runtime/TimeSynth/Source/TimeSynth/Classes/TimeSynthComponent.h
+++ b/Engine/Plugins/Runtime/TimeSynth/Source/TimeSynth/Classes/TimeSynthComponent.h
@@ -709,7 +709,7 @@ private:
 			: TargetVolumeDb(0.0f)
 			, StartVolumeDb(0.0f)
 			, CurrentVolumeDb(0.0f)
-			, LastVolumeDb(0.0f)
+			, LastVolumeDb(-1.0f)
 			, CurrentTime(0.0f)
 			, TargetFadeTime(0.0f)
 		{}
