Fixed crash if SetQuantizationSettings is called before the Time Synth is activated. (Added graceful handling and waring log about uninitialized Time Synth)

diff --git a/Engine/Plugins/Runtime/TimeSynth/Source/TimeSynth/Private/TimeSynthComponent.cpp b/Engine/Plugins/Runtime/TimeSynth/Source/TimeSynth/Private/TimeSynthComponent.cpp
--- a/Engine/Plugins/Runtime/TimeSynth/Source/TimeSynth/Private/TimeSynthComponent.cpp
+++ b/Engine/Plugins/Runtime/TimeSynth/Source/TimeSynth/Private/TimeSynthComponent.cpp
@@ -527,6 +527,12 @@ void UTimeSynthComponent::SetQuantizationSettings(const FTimeSynthQuantizationSe
 	// Store the quantization on the UObject for BP querying
 	QuantizationSettings = InQuantizationSettings;
 
+	// EventQuantizer will handle this case gracefully, but still warn.
+ 	if (!SampleRate)
+ 	{
+ 		UE_LOG(LogTimeSynth, Warning, TEXT("SetQuantizationSettings called with a sample rate of 0.  Did you forget to activate the time synth component?"));
+ 	}
+
 	// Local store what the global quantization is so we can assign it to clips using global quantization
 	GlobalQuantization = (Audio::EEventQuantization)InQuantizationSettings.GlobalQuantization;
 
diff --git a/Engine/Source/Runtime/SignalProcessing/Private/EventQuantizer.cpp b/Engine/Source/Runtime/SignalProcessing/Private/EventQuantizer.cpp
--- a/Engine/Source/Runtime/SignalProcessing/Private/EventQuantizer.cpp
+++ b/Engine/Source/Runtime/SignalProcessing/Private/EventQuantizer.cpp
@@ -53,7 +53,8 @@ namespace Audio
 
 		SetBPMInternal(QuantizationSettings.BeatsPerMinute);
 
-		if (!bQuantizationSettingsSet)
+		// allow these setting to cache, but do not reset event state until function is called with a valid sample rate
+		if (!bQuantizationSettingsSet && (InQuantizationSettings.SampleRate > 0))
 		{
 			bQuantizationSettingsSet = true;
 			ResetEventState();
diff --git a/Engine/Source/Runtime/SignalProcessing/Public/DSP/EventQuantizer.h b/Engine/Source/Runtime/SignalProcessing/Public/DSP/EventQuantizer.h
--- a/Engine/Source/Runtime/SignalProcessing/Public/DSP/EventQuantizer.h
+++ b/Engine/Source/Runtime/SignalProcessing/Public/DSP/EventQuantizer.h
@@ -47,7 +47,7 @@ namespace Audio
 		uint16 BeatDivision;
 
 		FEventQuantizationSettings()
-			: SampleRate(90.0f)
+			: SampleRate(0.0f)
 			, NumChannels(2)
 			, BeatsPerMinute(90.0f)
 			, BeatsPerBar(4)
