Fix for TimeSynth Volume Groups ignoring provided default volume

diff --git a/Engine/Plugins/Runtime/TimeSynth/Source/TimeSynth/Classes/TimeSynthComponent.h b/Engine/Plugins/Runtime/TimeSynth/Source/TimeSynth/Classes/TimeSynthComponent.h
--- a/Engine/Plugins/Runtime/TimeSynth/Source/TimeSynth/Classes/TimeSynthComponent.h
+++ b/Engine/Plugins/Runtime/TimeSynth/Source/TimeSynth/Classes/TimeSynthComponent.h
@@ -713,6 +713,15 @@ private:
 			, CurrentTime(0.0f)
 			, TargetFadeTime(0.0f)
 		{}
+
+		FVolumeGroupData(const float InitialVolume_dB)
+			: TargetVolumeDb(InitialVolume_dB)
+			, StartVolumeDb(InitialVolume_dB)
+			, CurrentVolumeDb(InitialVolume_dB)
+			, LastVolumeDb(InitialVolume_dB)
+			, CurrentTime(0.0f)
+			, TargetFadeTime(0.0f)
+		{}
 	};
 	void SetVolumeGroupInternal(FVolumeGroupData& InData, float VolumeDb, float FadeTimeSec);
 
diff --git a/Engine/Plugins/Runtime/TimeSynth/Source/TimeSynth/Private/TimeSynthComponent.cpp b/Engine/Plugins/Runtime/TimeSynth/Source/TimeSynth/Private/TimeSynthComponent.cpp
--- a/Engine/Plugins/Runtime/TimeSynth/Source/TimeSynth/Private/TimeSynthComponent.cpp
+++ b/Engine/Plugins/Runtime/TimeSynth/Source/TimeSynth/Private/TimeSynthComponent.cpp
@@ -717,16 +717,14 @@ FTimeSynthClipHandle UTimeSynthComponent::PlayClip(UTimeSynthClip* InClip, UTime
 		FVolumeGroupData* VolumeGroup = VolumeGroupData.Find(Id);
 		if (!VolumeGroup)
 		{
-			FVolumeGroupData NewData;
+			FVolumeGroupData NewData(InVolumeGroup->DefaultVolume);
 			NewData.Clips.Add(NewHandle);
-			VolumeGroupData.Add(Id, NewData);
-		}
-		else
-		{
-			// Get the current volume group value and "scale" it into the volume scale
-			VolumeScale *= Audio::ConvertToLinear(VolumeGroup->CurrentVolumeDb);
-			VolumeGroup->Clips.Add(NewHandle);
+			VolumeGroup = &VolumeGroupData.Add(Id, NewData);
 		}
+
+		// Get the current volume group value and "scale" it into the volume scale
+		VolumeScale *= Audio::ConvertToLinear(VolumeGroup->CurrentVolumeDb);
+		VolumeGroup->Clips.Add(NewHandle);
 	}
 
 	Audio::FSourceDecodeInit DecodeInit;
@@ -809,7 +807,7 @@ FTimeSynthClipHandle UTimeSynthComponent::PlayClip(UTimeSynthClip* InClip, UTime
 		}
 		else
 		{
-			PlayingClipsPool_AudioRenderThread[FreeClipIndex] = NewClipInfo;
+			PlayingClipsPool_AudioRenderThread[FreeClipIndex] = NewClipInfo;	
 		}
 
 		// Add a mapping of the clip handle id to the free index
