Fixes for TimeSynth GC crashes and bug where calling StopClip() on a queued clip that hasn't started wouldn't prevent the sound from playing.

diff --git a/Engine/Plugins/Runtime/TimeSynth/Source/TimeSynth/Classes/TimeSynthComponent.h b/Engine/Plugins/Runtime/TimeSynth/Source/TimeSynth/Classes/TimeSynthComponent.h
--- a/Engine/Plugins/Runtime/TimeSynth/Source/TimeSynth/Classes/TimeSynthComponent.h
+++ b/Engine/Plugins/Runtime/TimeSynth/Source/TimeSynth/Classes/TimeSynthComponent.h
@@ -569,6 +569,10 @@ public:
 	UFUNCTION(BlueprintCallable, Category = "Spectral Analysis", meta = (WorldContext = "WorldContextObject"))
 	void SetFFTSize(ETimeSynthFFTSize InFFTSize);
 
+	// Check to see if clips are actively generating sound on the TimeSynth
+	UFUNCTION(BlueprintCallable, Category = "Playback State", meta = (WorldContext = "WorldContextObject"))
+	bool HasActiveClips();
+
 private:
 	// Called when a new event happens when registered
 	void OnQuantizationEvent(Audio::EEventQuantization EventQuantizationType, int32 Bars, float Beat);
@@ -621,6 +625,12 @@ private:
 
 		bool bIsGloballyQuantized;
 
+		bool bIsInitialized;
+
+		bool bHasStartedPlaying;
+
+		bool bHasBeenStopped;
+
 		FPlayingClipInfo()
 			: ClipQuantization(Audio::EEventQuantization::Bar)
 			, VolumeScale(1.0f)
@@ -633,6 +643,9 @@ private:
 			, VolumeGroupId(INDEX_NONE)
 			, SynthClip(nullptr)
 			, bIsGloballyQuantized(false)
+			, bIsInitialized(false)
+			, bHasStartedPlaying(false)
+			, bHasBeenStopped(false)
 		{}
 	};
 
@@ -744,5 +757,7 @@ private:
 	// Need to limit output to prevent wrap around issues when converting to int16
 	Audio::FDynamicsProcessor DynamicsProcessor;
 
+	FThreadSafeBool bHasActiveClips;
+
 	friend class FTimeSynthEventListener;
 };
diff --git a/Engine/Plugins/Runtime/TimeSynth/Source/TimeSynth/Private/TimeSynthComponent.cpp b/Engine/Plugins/Runtime/TimeSynth/Source/TimeSynth/Private/TimeSynthComponent.cpp
--- a/Engine/Plugins/Runtime/TimeSynth/Source/TimeSynth/Private/TimeSynthComponent.cpp
+++ b/Engine/Plugins/Runtime/TimeSynth/Source/TimeSynth/Private/TimeSynthComponent.cpp
@@ -16,6 +16,7 @@ UTimeSynthComponent::UTimeSynthComponent(const FObjectInitializer& ObjectInitial
 	: Super(ObjectInitializer)
 	, MaxPoolSize(20)
 	, TimeSynthEventListener(this)
+	, bHasActiveClips(false)
 {
 	PrimaryComponentTick.bCanEverTick = true;
 }
@@ -177,6 +178,11 @@ void UTimeSynthComponent::SetFFTSize(ETimeSynthFFTSize InFFTSize)
 	});
 }
 
+bool UTimeSynthComponent::HasActiveClips()
+{
+	return bHasActiveClips;
+}
+
 void UTimeSynthComponent::OnQuantizationEvent(Audio::EEventQuantization EventQuantizationType, int32 Bars, float Beat)
 {
 	// When this happens, we want to queue up the event data so it can be safely consumed on the game thread
@@ -355,18 +361,18 @@ void UTimeSynthComponent::ShutdownPlayingClips()
 		int32 ClipIndex = ActivePlayingClipIndices_AudioRenderThread[i];
 		FPlayingClipInfo& PlayingClip = PlayingClipsPool_AudioRenderThread[ClipIndex];
 
-		// Block until the decoder has initialized
-		while (!SoundWaveDecoder.IsInitialized(PlayingClip.DecodingSoundSourceHandle))
+		// try to wait for the decoder to be initialized
+		// if we time out, this is probably in shutdown and the decoder won't ever init.
+		if(!SoundWaveDecoder.IsInitialized(PlayingClip.DecodingSoundSourceHandle))
 		{
-			FPlatformProcess::Sleep(0);
+			FPlatformProcess::Sleep(0.5f);
 		}
 
 		SoundWaveDecoder.RemoveDecodingSource(PlayingClip.DecodingSoundSourceHandle);
 		ActivePlayingClipIndices_AudioRenderThread.RemoveAtSwap(i, 1, false);
 		FreePlayingClipIndices_AudioRenderThread.Add(ClipIndex);
 	}
-
-	SoundWaveDecoder.Reset();
+	bHasActiveClips.AtomicSet(false);
 }
 
 void UTimeSynthComponent::OnEndGenerate() 
@@ -386,13 +392,17 @@ int32 UTimeSynthComponent::OnGenerateAudio(float* OutAudio, int32 NumSamples)
 	// to begin rendering. THe lambda callback will then enqueue any new rendering clips
 	// to the list of active clips. So we only need to loop through active clip indices to render the audio output
 	EventQuantizer.NotifyEvents(NumFrames);
+	const int32 NumActiveClips = ActivePlayingClipIndices_AudioRenderThread.Num();
+
+	bHasActiveClips.AtomicSet(NumActiveClips);
 
 	// Loop through all active loops and render their audio
-	for (int32 i = ActivePlayingClipIndices_AudioRenderThread.Num() - 1; i >= 0; --i)
+	for (int32 i = NumActiveClips - 1; i >= 0; --i)
 	{
 		// Grab the playing clip at the active index
 		int32 ClipIndex = ActivePlayingClipIndices_AudioRenderThread[i];
 		FPlayingClipInfo& PlayingClip = PlayingClipsPool_AudioRenderThread[ClipIndex];
+		PlayingClip.bHasStartedPlaying = true;
 
 		// Compute the number of frames we need to read
 		int32 NumFramesToRead = NumFrames - PlayingClip.StartFrameOffset;
@@ -733,6 +743,7 @@ FTimeSynthClipHandle UTimeSynthComponent::PlayClip(UTimeSynthClip* InClip, UTime
 	DecodeInit.VolumeScale = VolumeScale;
 	DecodeInit.SoundWave = ChosenSound.SoundWave;
 	DecodeInit.SeekTime = 0;
+	DecodeInit.bForceSyncDecode = true;
 
 	// Update the synth component on the audio thread
 	FAudioThread::RunCommandOnAudioThread([this, DecodeInit]()
@@ -801,13 +812,15 @@ FTimeSynthClipHandle UTimeSynthComponent::PlayClip(UTimeSynthClip* InClip, UTime
 		if(bCurrentPoolSizeChanged)
 		{
 			PlayingClipsPool_AudioRenderThread.Add(NewClipInfo);
+			PlayingClipsPool_AudioRenderThread[PlayingClipsPool_AudioRenderThread.Num() - 1].bIsInitialized = true;
 
 			UE_LOG(LogTimeSynth, Warning, TEXT("Reallocating PlayingClipsPool to %i (which is a performance hit.) If this wasn't caused by a hitch, consider initializing Pool Size to a larger value.")
 				, PlayingClipsPool_AudioRenderThread.Num());
 		}
 		else
 		{
-			PlayingClipsPool_AudioRenderThread[FreeClipIndex] = NewClipInfo;	
+			PlayingClipsPool_AudioRenderThread[FreeClipIndex] = NewClipInfo;
+			PlayingClipsPool_AudioRenderThread[FreeClipIndex].bIsInitialized = true;
 		}
 
 		// Add a mapping of the clip handle id to the free index
@@ -824,6 +837,12 @@ FTimeSynthClipHandle UTimeSynthComponent::PlayClip(UTimeSynthClip* InClip, UTime
 			{
 				FPlayingClipInfo& PlayingClipInfo = PlayingClipsPool_AudioRenderThread[FreeClipIndex];
 
+				// early exit.  Quantized stop event was just proccessed for this same Quantization step
+				if (PlayingClipInfo.bHasBeenStopped)
+				{
+					return;
+				}
+
 				// Setup the duration of various things using the event quantizer
 				PlayingClipInfo.DurationFrames = EventQuantizer.GetDurationInFrames(ClipDuration.NumBars, (float)ClipDuration.NumBeats);
 				PlayingClipInfo.FadeInDurationFrames = EventQuantizer.GetDurationInFrames(FadeInTime.NumBars, (float)FadeInTime.NumBeats);
@@ -854,14 +873,43 @@ void UTimeSynthComponent::StopClip(FTimeSynthClipHandle InClipHandle, ETimeSynth
 
 			[this, InClipHandle](uint32 NumFramesOffset)
 			{
-				int32* PlayingClipIndex = ClipIdToClipIndexMap_AudioRenderThread.Find(InClipHandle.ClipId);
-				if (PlayingClipIndex)
+				int32* PlayingClipIndexPtr = ClipIdToClipIndexMap_AudioRenderThread.Find(InClipHandle.ClipId);
+				if (PlayingClipIndexPtr)
 				{
+					const int32 PlayingClipIndex = *PlayingClipIndexPtr;
 					// Grab the clip info
-					FPlayingClipInfo& PlayingClipInfo = PlayingClipsPool_AudioRenderThread[*PlayingClipIndex];
+					FPlayingClipInfo& PlayingClipInfo = PlayingClipsPool_AudioRenderThread[PlayingClipIndex];
 
-					// Only do anything if the clip is not yet already fading
-					if (PlayingClipInfo.CurrentFrameCount < PlayingClipInfo.DurationFrames)
+					if (!PlayingClipInfo.bHasStartedPlaying)
+					{
+						// add index back to free pool
+						FreePlayingClipIndices_AudioRenderThread.Add(PlayingClipIndex);
+
+						// remove map entry
+						ClipIdToClipIndexMap_AudioRenderThread.Remove(InClipHandle.ClipId);
+
+						// Clip may have been staged to play, so we can search and remove it manually
+						bool bFoundClip = false;
+						int32 NumActivePlayingClipIndices = ActivePlayingClipIndices_AudioRenderThread.Num();
+
+						for (int i = 0; i < NumActivePlayingClipIndices; ++i)
+						{
+							if (ActivePlayingClipIndices_AudioRenderThread[i] == PlayingClipIndex)
+							{
+								ActivePlayingClipIndices_AudioRenderThread.RemoveAtSwap(i, 1, false);
+								bFoundClip = true;
+								--NumActivePlayingClipIndices;
+							}
+						}
+
+						// clip wasn't staged to play yet. Raise the flag for the Notify
+						if (!bFoundClip)
+						{
+							PlayingClipInfo.bHasBeenStopped = true;
+						}
+					}
+					// Already playing the clip, so force the clip to start fading if its already playing
+					else if (PlayingClipInfo.CurrentFrameCount < PlayingClipInfo.DurationFrames)
 					{
 						// Adjust the duration of the clip to "spoof" it's code which triggers a fade this render callback block.
 						PlayingClipInfo.DurationFrames = PlayingClipInfo.CurrentFrameCount + NumFramesOffset;
diff --git a/Engine/Source/Runtime/AudioMixer/Private/AudioMixerSourceBuffer.cpp b/Engine/Source/Runtime/AudioMixer/Private/AudioMixerSourceBuffer.cpp
--- a/Engine/Source/Runtime/AudioMixer/Private/AudioMixerSourceBuffer.cpp
+++ b/Engine/Source/Runtime/AudioMixer/Private/AudioMixerSourceBuffer.cpp
@@ -4,6 +4,7 @@
 #include "AudioMixerSourceDecode.h"
 #include "ContentStreaming.h"
 #include "AudioDecompress.h"
+#include "Misc/ScopeTryLock.h"
 
 namespace Audio
 {
@@ -110,6 +111,10 @@ namespace Audio
 
 	FMixerSourceBuffer::~FMixerSourceBuffer()
 	{
+		// GC methods may get called from the game thread during the destructor
+		// These methods will trylock and early exit if we have this lock
+		FScopeLock Lock(&DtorCritSec);
+
 		// Make sure we have completed our async realtime task before deleting the decompression state
 		if (AsyncRealtimeAudioTask)
 		{
@@ -452,9 +457,18 @@ namespace Audio
 		BufferQueue.Enqueue(InSourceVoiceBuffer);
 	}
 
-	void FMixerSourceBuffer::OnBeginDestroy(USoundWave * /*Wave*/)
+	bool FMixerSourceBuffer::OnBeginDestroy(USoundWave * /*Wave*/)
 	{
-		SoundWave = nullptr;
+		FScopeTryLock Lock(&DtorCritSec);
+
+		// if we don't have the lock, it means we are in ~FMixerSourceBuffer() on another thread
+		if (Lock.IsLocked() && SoundWave)
+		{
+			SoundWave = nullptr;
+			return true;
+		}
+		
+		return false;
 	}
 
 	bool FMixerSourceBuffer::OnIsReadyForFinishDestroy(USoundWave * /*Wave*/) const
@@ -464,7 +478,13 @@ namespace Audio
 
 	void FMixerSourceBuffer::OnFinishDestroy(USoundWave * /*Wave*/)
 	{
-		SoundWave = nullptr;
+		FScopeTryLock Lock(&DtorCritSec);
+
+		// if we don't have the lock, it means we are in ~FMixerSourceBuffer() on another thread
+		if (Lock.IsLocked() && SoundWave)
+		{
+			SoundWave = nullptr;
+		}
 	}
 
 	bool FMixerSourceBuffer::IsAsyncTaskInProgress() const
diff --git a/Engine/Source/Runtime/AudioMixer/Private/AudioMixerSourceBuffer.h b/Engine/Source/Runtime/AudioMixer/Private/AudioMixerSourceBuffer.h
--- a/Engine/Source/Runtime/AudioMixer/Private/AudioMixerSourceBuffer.h
+++ b/Engine/Source/Runtime/AudioMixer/Private/AudioMixerSourceBuffer.h
@@ -121,6 +121,8 @@ namespace Audio
 		int32 NumPrecacheFrames;
 		TArray<uint8> CachedRealtimeFirstBuffer;
 
+		FCriticalSection DtorCritSec;
+
 		uint32 bInitialized : 1;
 		uint32 bBufferFinished : 1;
 		uint32 bPlayedCachedBuffer : 1;
@@ -130,7 +132,7 @@ namespace Audio
 		uint32 bIsBus : 1;
 		uint32 bForceSyncDecode : 1;
 	
-		virtual void OnBeginDestroy(class USoundWave* Wave) override;
+		virtual bool OnBeginDestroy(class USoundWave* Wave) override;
 		virtual bool OnIsReadyForFinishDestroy(class USoundWave* Wave) const override;
 		virtual void OnFinishDestroy(class USoundWave* Wave) override;
 	};
diff --git a/Engine/Source/Runtime/AudioMixer/Private/SoundWaveDecoder.cpp b/Engine/Source/Runtime/AudioMixer/Private/SoundWaveDecoder.cpp
--- a/Engine/Source/Runtime/AudioMixer/Private/SoundWaveDecoder.cpp
+++ b/Engine/Source/Runtime/AudioMixer/Private/SoundWaveDecoder.cpp
@@ -361,11 +361,6 @@ namespace Audio
 		return true;
 	}
 
-	void FDecodingSoundSource::AddReferencedObjects(FReferenceCollector & Collector)
-	{
-		Collector.AddReferencedObject(SoundWave);
-	}
-
 	FSoundSourceDecoder::FSoundSourceDecoder()
 		: AudioThreadId(0)
 		, AudioDevice(nullptr)
@@ -378,6 +373,30 @@ namespace Audio
 		
 	}
 
+	void FSoundSourceDecoder::AddReferencedObjects(FReferenceCollector & Collector)
+	{
+		for (auto& Entry : PrecachingSources)
+		{
+			FSourceDecodeInit& DecodingSoundInitPtr = Entry.Value;
+			Collector.AddReferencedObject(DecodingSoundInitPtr.SoundWave);
+		}
+
+		for (auto& Entry : InitializingDecodingSources)
+		{
+			FDecodingSoundSourcePtr DecodingSoundSourcePtr = Entry.Value;
+			USoundWave* SoundWave = DecodingSoundSourcePtr->GetSoundWave();
+			Collector.AddReferencedObject(SoundWave);
+		}
+
+		FScopeLock Lock(&DecodingSourcesCritSec);
+		for (auto& Entry : DecodingSources)
+		{
+			FDecodingSoundSourcePtr DecodingSoundSourcePtr = Entry.Value;
+			USoundWave* SoundWave = DecodingSoundSourcePtr->GetSoundWave();
+			Collector.AddReferencedObject(SoundWave);
+		}
+	}
+
 	void FSoundSourceDecoder::Init(FAudioDevice* InAudioDevice, int32 InSampleRate)
 	{
 		AudioDevice = InAudioDevice;
@@ -422,6 +441,7 @@ namespace Audio
 			// Add this decoding sound wave to a data structure we can access safely from audio render thread
 			EnqueueDecoderCommand([this, InitData, DecodingSoundWaveDataPtr]()
 			{
+				FScopeLock Lock(&DecodingSourcesCritSec);
 				DecodingSources.Add(InitData.Handle.Id, DecodingSoundWaveDataPtr);
 
 				UE_LOG(LogAudioMixer, Verbose, TEXT("Decoding SoundWave '%s' (Num Decoding: %d)"),
@@ -489,6 +509,7 @@ namespace Audio
 
 	void FSoundSourceDecoder::RemoveDecodingSource(const FDecodingSoundSourceHandle& Handle)
 	{
+		FScopeLock Lock(&DecodingSourcesCritSec);
 		DecodingSources.Remove(Handle.Id);
 	}
 
diff --git a/Engine/Source/Runtime/AudioMixer/Public/SoundWaveDecoder.h b/Engine/Source/Runtime/AudioMixer/Public/SoundWaveDecoder.h
--- a/Engine/Source/Runtime/AudioMixer/Public/SoundWaveDecoder.h
+++ b/Engine/Source/Runtime/AudioMixer/Public/SoundWaveDecoder.h
@@ -9,6 +9,8 @@
 #include "AudioMixer/Private/AudioMixerSourceBuffer.h"
 #include "DSP/SinOsc.h"
 #include "DSP/ParamInterpolator.h"
+#include "Misc/ScopeLock.h"
+
 
 class FAudioDevice;
 
@@ -44,7 +46,7 @@ namespace Audio
 		FDecodingSoundSourceHandle Handle;
 	};
 
-	class FDecodingSoundSource : public FGCObject
+	class FDecodingSoundSource
 	{
 	public:
 		FDecodingSoundSource(FAudioDevice* AudioDevice, const FSourceDecodeInit& InitData);
@@ -77,8 +79,8 @@ namespace Audio
 		// Get audio buffer
 		bool GetAudioBuffer(const int32 InNumFrames, const int32 InNumChannels, AlignedFloatBuffer& OutAudioBuffer);
 
-		// Override for FGCObject:
-		virtual void AddReferencedObjects(FReferenceCollector & Collector) override;
+		// Return the underlying sound wave
+		USoundWave* GetSoundWave() { return SoundWave; }
 
 	private:
 
@@ -196,12 +198,16 @@ namespace Audio
 
 	typedef TSharedPtr<FDecodingSoundSource> FDecodingSoundSourcePtr;
 	
-	class AUDIOMIXER_API FSoundSourceDecoder
+	class AUDIOMIXER_API FSoundSourceDecoder : public FGCObject
 	{
 	public:
 		FSoundSourceDecoder();
 		virtual ~FSoundSourceDecoder();
 
+		//~ Begin FGCObject
+		virtual void AddReferencedObjects(FReferenceCollector & Collector) override;
+		//~ End FGCObject
+
 		// Initialize the source decoder at the given output sample rate
 		// Sources will automatically sample rate convert to match this output
 		void Init(FAudioDevice* InAudioDevice, int32 SampleRate);
@@ -247,6 +253,7 @@ namespace Audio
 		int32 AudioThreadId;
 		FAudioDevice* AudioDevice;
 		int32 SampleRate;
+		FCriticalSection DecodingSourcesCritSec;
 		TMap<int32, FDecodingSoundSourcePtr> InitializingDecodingSources;
 		TMap<int32, FDecodingSoundSourcePtr> DecodingSources;
 		TMap<int32, FSourceDecodeInit> PrecachingSources;
diff --git a/Engine/Source/Runtime/Engine/Classes/Sound/SoundWave.h b/Engine/Source/Runtime/Engine/Classes/Sound/SoundWave.h
--- a/Engine/Source/Runtime/Engine/Classes/Sound/SoundWave.h
+++ b/Engine/Source/Runtime/Engine/Classes/Sound/SoundWave.h
@@ -288,7 +288,8 @@ struct ISoundWaveClient
 	ISoundWaveClient() {}
 	virtual ~ISoundWaveClient() {}
 	
-	virtual void OnBeginDestroy(class USoundWave* Wave) = 0;
+	// OnBeginDestroy() returns true to unsubscribe as an ISoundWaveClient
+	virtual bool OnBeginDestroy(class USoundWave* Wave) = 0;
 	virtual bool OnIsReadyForFinishDestroy(class USoundWave* Wave) const = 0;
 	virtual void OnFinishDestroy(class USoundWave* Wave) = 0;
 };
diff --git a/Engine/Source/Runtime/Engine/Private/SoundWave.cpp b/Engine/Source/Runtime/Engine/Private/SoundWave.cpp
--- a/Engine/Source/Runtime/Engine/Private/SoundWave.cpp
+++ b/Engine/Source/Runtime/Engine/Private/SoundWave.cpp
@@ -948,9 +948,17 @@ void USoundWave::BeginDestroy()
 	
 	{
 		FScopeLock Lock(&SourcesPlayingCs);
-		for (ISoundWaveClient* i: SourcesPlaying)
+		int32 CurrNumSourcesPlaying = SourcesPlaying.Num();
+
+		for (int32 i = CurrNumSourcesPlaying - 1; i >= 0; --i)
 		{
-			i->OnBeginDestroy(this);
+			ISoundWaveClient* SoundWaveClientPtr = SourcesPlaying[i];
+
+			if (SoundWaveClientPtr && SoundWaveClientPtr->OnBeginDestroy(this))
+			{
+				// if OnBeginDestroy returned true, we are unsubscribing the SoundWaveClient...
+				SourcesPlaying.RemoveAtSwap(i, 1, false);
+			}
 		}
 	}
 
@@ -1620,12 +1628,11 @@ bool USoundWave::IsReadyForFinishDestroy()
 		
 		for (ISoundWaveClient* SoundWaveClientPtr : SourcesPlaying)
 		{
-			if(SoundWaveClientPtr)
-			{
-				SoundWaveClientPtr->OnIsReadyForFinishDestroy(this);
-			}
+			// Should never have a null element here, inactive elements are removed from the container
+			check(SoundWaveClientPtr);
+			SoundWaveClientPtr->OnIsReadyForFinishDestroy(this);
 		}
-	}
+	} 
 
 	const bool bIsStreamingInProgress = IStreamingManager::Get().GetAudioStreamingManager().IsStreamingInProgress(this);
 
@@ -1652,6 +1659,20 @@ void USoundWave::FinishDestroy()
 {
 	Super::FinishDestroy();
 
+	FScopeLock Lock(&SourcesPlayingCs);
+	int32 CurrNumSourcesPlaying = SourcesPlaying.Num();
+
+	for (int32 i = CurrNumSourcesPlaying - 1; i >= 0; --i)
+	{
+		ISoundWaveClient* SoundWaveClientPtr = SourcesPlaying[i];
+
+		if (SoundWaveClientPtr)
+		{
+			SoundWaveClientPtr->OnFinishDestroy(this);
+			SourcesPlaying.RemoveAtSwap(i, 1, false);
+		}
+	}
+
 	check(GetPrecacheState() != ESoundWavePrecacheState::InProgress);
 	check(AudioDecompressor == nullptr);
 
