Unshelved strict compile fix from Max

diff --git a/Engine/Plugins/Runtime/TimeSynth/Source/TimeSynth/Private/TimeSynthComponent.cpp b/Engine/Plugins/Runtime/TimeSynth/Source/TimeSynth/Private/TimeSynthComponent.cpp
--- a/Engine/Plugins/Runtime/TimeSynth/Source/TimeSynth/Private/TimeSynthComponent.cpp
+++ b/Engine/Plugins/Runtime/TimeSynth/Source/TimeSynth/Private/TimeSynthComponent.cpp
@@ -394,7 +394,7 @@ int32 UTimeSynthComponent::OnGenerateAudio(float* OutAudio, int32 NumSamples)
 	EventQuantizer.NotifyEvents(NumFrames);
 	const int32 NumActiveClips = ActivePlayingClipIndices_AudioRenderThread.Num();
 
-	bHasActiveClips.AtomicSet(NumActiveClips);
+	bHasActiveClips.AtomicSet(NumActiveClips > 0);
 
 	// Loop through all active loops and render their audio
 	for (int32 i = NumActiveClips - 1; i >= 0; --i)
