Adding clamp to TimeSynth BPM (both in Ustruct and BP interface)

diff --git a/Engine/Plugins/Runtime/TimeSynth/Source/TimeSynth/Classes/TimeSynthComponent.h b/Engine/Plugins/Runtime/TimeSynth/Source/TimeSynth/Classes/TimeSynthComponent.h
--- a/Engine/Plugins/Runtime/TimeSynth/Source/TimeSynth/Classes/TimeSynthComponent.h
+++ b/Engine/Plugins/Runtime/TimeSynth/Source/TimeSynth/Classes/TimeSynthComponent.h
@@ -132,7 +132,7 @@ struct TIMESYNTH_API FTimeSynthQuantizationSettings
 	GENERATED_USTRUCT_BODY()
 
 	// The beats per minute of the pulse. Musical convention gives this as BPM for "quarter notes" (BeatDivision = 4).
-	UPROPERTY(EditAnywhere, BlueprintReadWrite, Category = "Synth|TimeSynth|PlayClip", meta = (ClampMin = "1.0", UIMin = "1.0"))
+	UPROPERTY(EditAnywhere, BlueprintReadWrite, Category = "Synth|TimeSynth|PlayClip", meta = (ClampMin = "1.0", UIMin = "1.0", ClampMax = "999.0", UIMax = "999.0"))
 	float BeatsPerMinute;
 
 	// Defines numerator when determining beat time in seconds
diff --git a/Engine/Plugins/Runtime/TimeSynth/Source/TimeSynth/Private/TimeSynthComponent.cpp b/Engine/Plugins/Runtime/TimeSynth/Source/TimeSynth/Private/TimeSynthComponent.cpp
--- a/Engine/Plugins/Runtime/TimeSynth/Source/TimeSynth/Private/TimeSynthComponent.cpp
+++ b/Engine/Plugins/Runtime/TimeSynth/Source/TimeSynth/Private/TimeSynthComponent.cpp
@@ -563,11 +563,18 @@ void UTimeSynthComponent::SetQuantizationSettings(const FTimeSynthQuantizationSe
 
 void UTimeSynthComponent::SetBPM(const float InBeatsPerMinute)
 {
-	QuantizationSettings.BeatsPerMinute = InBeatsPerMinute;
+	float NewBeatsPerMinute = FMath::Clamp(InBeatsPerMinute, 0.0f, 999.0f);
 
-	SynthCommand([this, InBeatsPerMinute]
+	if (!FMath::IsNearlyEqual(NewBeatsPerMinute, InBeatsPerMinute))
 	{
-		EventQuantizer.SetBPM(InBeatsPerMinute);
+		UE_LOG(LogTimeSynth, Warning, TEXT("Clapming provided BPM from %f to %f"), InBeatsPerMinute, NewBeatsPerMinute);
+	}
+
+	QuantizationSettings.BeatsPerMinute = NewBeatsPerMinute;
+
+	SynthCommand([this, NewBeatsPerMinute]
+	{
+		EventQuantizer.SetBPM(NewBeatsPerMinute);
 	});
 }
 
