Made TimeSynth hot-fix-off-able (au.DisableTimeSynth)

diff --git a/Engine/Plugins/Runtime/TimeSynth/Source/TimeSynth/Classes/TimeSynthComponent.h b/Engine/Plugins/Runtime/TimeSynth/Source/TimeSynth/Classes/TimeSynthComponent.h
--- a/Engine/Plugins/Runtime/TimeSynth/Source/TimeSynth/Classes/TimeSynthComponent.h
+++ b/Engine/Plugins/Runtime/TimeSynth/Source/TimeSynth/Classes/TimeSynthComponent.h
@@ -758,6 +758,7 @@ private:
 	Audio::FDynamicsProcessor DynamicsProcessor;
 
 	FThreadSafeBool bHasActiveClips;
+	FThreadSafeBool bTimeSynthWasDisabled;
 
 	friend class FTimeSynthEventListener;
 };
diff --git a/Engine/Plugins/Runtime/TimeSynth/Source/TimeSynth/Private/TimeSynthComponent.cpp b/Engine/Plugins/Runtime/TimeSynth/Source/TimeSynth/Private/TimeSynthComponent.cpp
--- a/Engine/Plugins/Runtime/TimeSynth/Source/TimeSynth/Private/TimeSynthComponent.cpp
+++ b/Engine/Plugins/Runtime/TimeSynth/Source/TimeSynth/Private/TimeSynthComponent.cpp
@@ -3,6 +3,14 @@
 #include "TimeSynthComponent.h"
 #include "AudioThread.h"
 #include "TimeSynthModule.h"
+
+static int32 DisableTimeSynthCvar = 0;
+FAutoConsoleVariableRef CVarDisableTimeSynth(
+	TEXT("au.DisableTimeSynth"),
+	DisableTimeSynthCvar,
+	TEXT("Disables all TimeSynth rendering/processing.\n")
+	TEXT("0: Not Disabled, 1: Disabled"),
+	ECVF_Default);
  
 static_assert((int32)Audio::EEventQuantization::Count == (int32)ETimeSynthEventQuantization::Count, "These enumerations need to match");
 
@@ -17,6 +25,7 @@ UTimeSynthComponent::UTimeSynthComponent(const FObjectInitializer& ObjectInitial
 	, MaxPoolSize(20)
 	, TimeSynthEventListener(this)
 	, bHasActiveClips(false)
+	, bTimeSynthWasDisabled(false)
 {
 	PrimaryComponentTick.bCanEverTick = true;
 }
@@ -93,6 +102,11 @@ void UTimeSynthComponent::AddQuantizationEventDelegate(ETimeSynthEventQuantizati
 
 void UTimeSynthComponent::SetFilterSettings(ETimeSynthFilter InFilter, const FTimeSynthFilterSettings& InSettings)
 {
+	if (DisableTimeSynthCvar || bTimeSynthWasDisabled)
+	{
+		return;
+	}
+
 	if (InFilter == ETimeSynthFilter::FilterA)
 	{
 		FilterASettings = InSettings;
@@ -111,6 +125,11 @@ void UTimeSynthComponent::SetFilterSettings(ETimeSynthFilter InFilter, const FTi
 
 void UTimeSynthComponent::SetEnvelopeFollowerSettings(const FTimeSynthEnvelopeFollowerSettings& InSettings)
 {
+	if (DisableTimeSynthCvar || bTimeSynthWasDisabled)
+	{
+		return;
+	}
+
 	EnvelopeFollowerSettings = InSettings;
 
 	SynthCommand([this, InSettings]
@@ -122,6 +141,11 @@ void UTimeSynthComponent::SetEnvelopeFollowerSettings(const FTimeSynthEnvelopeFo
 
 void UTimeSynthComponent::SetFilterEnabled(ETimeSynthFilter InFilter, bool bInIsFilterEnabled)
 {
+	if (DisableTimeSynthCvar || bTimeSynthWasDisabled)
+	{
+		return;
+	}
+
 	if (InFilter == ETimeSynthFilter::FilterA)
 	{
 		bIsFilterAEnabled = bInIsFilterEnabled;
@@ -139,6 +163,11 @@ void UTimeSynthComponent::SetFilterEnabled(ETimeSynthFilter InFilter, bool bInIs
 
 void UTimeSynthComponent::SetEnvelopeFollowerEnabled(bool bInIsEnabled)
 {
+	if (DisableTimeSynthCvar || bTimeSynthWasDisabled)
+	{
+		return;
+	}
+
 	bIsEnvelopeFollowerEnabled = bInIsEnabled;
 
 	// Set the envelope value to 0.0 immediately if we're disabling the envelope follower
@@ -169,6 +198,11 @@ Audio::FSpectrumAnalyzerSettings::EFFTSize UTimeSynthComponent::GetFFTSize(ETime
 
 void UTimeSynthComponent::SetFFTSize(ETimeSynthFFTSize InFFTSize)
 {
+	if (DisableTimeSynthCvar || bTimeSynthWasDisabled)
+	{
+		return;
+	}
+
 	Audio::FSpectrumAnalyzerSettings::EFFTSize NewFFTSize = GetFFTSize(InFFTSize);
 
 	SynthCommand([this, NewFFTSize]
@@ -185,6 +219,11 @@ bool UTimeSynthComponent::HasActiveClips()
 
 void UTimeSynthComponent::OnQuantizationEvent(Audio::EEventQuantization EventQuantizationType, int32 Bars, float Beat)
 {
+	if (DisableTimeSynthCvar || bTimeSynthWasDisabled)
+	{
+		return;
+	}
+
 	// When this happens, we want to queue up the event data so it can be safely consumed on the game thread
 	GameCommand([this, EventQuantizationType, Bars, Beat]()
 	{
@@ -211,6 +250,14 @@ void UTimeSynthComponent::TickComponent(float DeltaTime, enum ELevelTick TickTyp
 	// Pump the command queue for any event data that is coming back from the audio render thread/callback
 	PumpGameCommandQueue();
 
+	if (DisableTimeSynthCvar || bTimeSynthWasDisabled)
+	{
+		OnEndGenerate();
+		SetActive(false);
+		bTimeSynthWasDisabled = static_cast<bool>(DisableTimeSynthCvar);
+		return;
+	}
+
 	// Broadcast the playback time
 	if (OnPlaybackTime.IsBound())
 	{
@@ -300,7 +347,15 @@ void UTimeSynthComponent::UpdateEnvelopeFollower()
 }
 
 bool UTimeSynthComponent::Init(int32& InSampleRate)
-{ 
+{
+	if (DisableTimeSynthCvar || bTimeSynthWasDisabled)
+	{
+		SetActive(false);
+		OnEndGenerate();
+		bTimeSynthWasDisabled = static_cast<bool>(DisableTimeSynthCvar);
+		return false;
+	}
+
 	SampleRate = InSampleRate;
 	SoundWaveDecoder.Init(GetAudioDevice(), InSampleRate);
 	NumChannels = 2;
@@ -372,6 +427,11 @@ void UTimeSynthComponent::ShutdownPlayingClips()
 		ActivePlayingClipIndices_AudioRenderThread.RemoveAtSwap(i, 1, false);
 		FreePlayingClipIndices_AudioRenderThread.Add(ClipIndex);
 	}
+	ActivePlayingClipIndices_AudioRenderThread.Reset();
+	FreePlayingClipIndices_AudioRenderThread.Reset();
+	ClipIdToClipIndexMap_AudioRenderThread.Reset();
+	DecodingSounds_GameThread.Reset();
+
 	bHasActiveClips.AtomicSet(false);
 }
 
@@ -382,6 +442,12 @@ void UTimeSynthComponent::OnEndGenerate()
 
 int32 UTimeSynthComponent::OnGenerateAudio(float* OutAudio, int32 NumSamples)
 {
+	if (DisableTimeSynthCvar || bTimeSynthWasDisabled)
+	{
+		bTimeSynthWasDisabled = static_cast<bool>(DisableTimeSynthCvar);
+		return 0;
+	}
+
 	// Update the decoder
 	SoundWaveDecoder.UpdateRenderThread();
 
@@ -534,6 +600,11 @@ int32 UTimeSynthComponent::OnGenerateAudio(float* OutAudio, int32 NumSamples)
 
 void UTimeSynthComponent::SetQuantizationSettings(const FTimeSynthQuantizationSettings& InQuantizationSettings)
 {
+	if (DisableTimeSynthCvar || bTimeSynthWasDisabled)
+	{
+		return;
+	}
+
 	// Store the quantization on the UObject for BP querying
 	QuantizationSettings = InQuantizationSettings;
 
@@ -595,6 +666,17 @@ void UTimeSynthComponent::ResetSeed()
 
 FTimeSynthClipHandle UTimeSynthComponent::PlayClip(UTimeSynthClip* InClip, UTimeSynthVolumeGroup* InVolumeGroup)
 {
+	if (DisableTimeSynthCvar || !GetAudioDevice() || bTimeSynthWasDisabled)
+	{
+		SetActive(false);
+		if (!bTimeSynthWasDisabled)
+		{
+			OnEndGenerate();
+		}
+		bTimeSynthWasDisabled = true;
+		return FTimeSynthClipHandle();
+	}
+
 	if (!InClip)
 	{
 		UE_LOG(LogTimeSynth, Warning, TEXT("Failed to play clip. Null UTimeSynthClip object."));
