Fix for TimeSynth volume groups only getting applied to a single instance of an active clip

diff --git a/Engine/Plugins/Runtime/TimeSynth/Source/TimeSynth/Private/TimeSynthComponent.cpp b/Engine/Plugins/Runtime/TimeSynth/Source/TimeSynth/Private/TimeSynthComponent.cpp
--- a/Engine/Plugins/Runtime/TimeSynth/Source/TimeSynth/Private/TimeSynthComponent.cpp
+++ b/Engine/Plugins/Runtime/TimeSynth/Source/TimeSynth/Private/TimeSynthComponent.cpp
@@ -284,13 +284,13 @@ void UTimeSynthComponent::TickComponent(float DeltaTime, enum ELevelTick TickTyp
 			VolumeGroup.CurrentTime += DeltaTime;
 		}
 
-		for (FTimeSynthClipHandle& ClipHandle : VolumeGroup.Clips)
+		if (FMath::IsNearlyEqual(VolumeGroup.CurrentVolumeDb, VolumeGroup.LastVolumeDb, KINDA_SMALL_NUMBER))
 		{
-			if (FMath::IsNearlyEqual(VolumeGroup.CurrentVolumeDb, VolumeGroup.LastVolumeDb, KINDA_SMALL_NUMBER))
-			{
-				continue;
-			}
+			continue;
+		}
 
+		for (FTimeSynthClipHandle& ClipHandle : VolumeGroup.Clips)
+		{
 			VolumeGroup.LastVolumeDb = VolumeGroup.CurrentVolumeDb;
 
 			float LinearVolume = Audio::ConvertToLinear(VolumeGroup.CurrentVolumeDb);
@@ -820,10 +820,13 @@ FTimeSynthClipHandle UTimeSynthComponent::PlayClip(UTimeSynthClip* InClip, UTime
 			NewData.Clips.Add(NewHandle);
 			VolumeGroup = &VolumeGroupData.Add(Id, NewData);
 		}
+		else
+		{
+			VolumeGroup->Clips.Add(NewHandle);
+		}
 
 		// Get the current volume group value and "scale" it into the volume scale
 		VolumeScale *= Audio::ConvertToLinear(VolumeGroup->CurrentVolumeDb);
-		VolumeGroup->Clips.Add(NewHandle);
 	}
 
 	Audio::FSourceDecodeInit DecodeInit;
