Adding cvars and logging to time synth

diff --git a/Engine/Plugins/Runtime/TimeSynth/Source/TimeSynth/Private/TimeSynthComponent.cpp b/Engine/Plugins/Runtime/TimeSynth/Source/TimeSynth/Private/TimeSynthComponent.cpp
--- a/Engine/Plugins/Runtime/TimeSynth/Source/TimeSynth/Private/TimeSynthComponent.cpp
+++ b/Engine/Plugins/Runtime/TimeSynth/Source/TimeSynth/Private/TimeSynthComponent.cpp
@@ -11,6 +11,30 @@ FAutoConsoleVariableRef CVarDisableTimeSynth(
 	TEXT("Disables all TimeSynth rendering/processing.\n")
 	TEXT("0: Not Disabled, 1: Disabled"),
 	ECVF_Default);
+
+static int32 TimeSynthForceSyncDecodesCvar = 1;
+FAutoConsoleVariableRef CVarTimeSynthForceSyncDecodes(
+	TEXT("au.TimeSynthForceSyncDecodes"),
+	TimeSynthForceSyncDecodesCvar,
+	TEXT("Forces decodes of TimeSynth audio to be synchronous.\n")
+	TEXT("0: Async decodes, 1: Sync decodes"),
+	ECVF_Default);
+
+static int32 TimeSynthCallbackSizeOverrideCvar = -1;
+FAutoConsoleVariableRef TimeSynthCallbackSizeOverride(
+	TEXT("au.TimeSynthCallbackSizeOverride"),
+	TimeSynthCallbackSizeOverrideCvar,
+	TEXT("Overrides the default callback size of synth components for Time Synths.\n")
+	TEXT(" <= default: Don't override, > default: callback size to use, (must be a multiple of 4)"),
+	ECVF_Default);
+
+static int32 TimeSynthActiveClipLimitCvar = 20;
+FAutoConsoleVariableRef TimeSynthActiveClipLimit(
+	TEXT("au.TimeSynthActiveClipLimit"),
+	TimeSynthActiveClipLimitCvar,
+	TEXT("Overrides the default callback size of synth components for Time Synths.\n")
+	TEXT(" <= default: Don't override, > default: callback size to use, (must be a multiple of 4)"),
+	ECVF_Default);
  
 static_assert((int32)Audio::EEventQuantization::Count == (int32)ETimeSynthEventQuantization::Count, "These enumerations need to match");
 
@@ -22,12 +46,17 @@ void FTimeSynthEventListener::OnEvent(Audio::EEventQuantization EventQuantizatio
 
 UTimeSynthComponent::UTimeSynthComponent(const FObjectInitializer& ObjectInitializer) 
 	: Super(ObjectInitializer)
-	, MaxPoolSize(20)
+	, MaxPoolSize(TimeSynthActiveClipLimitCvar)
 	, TimeSynthEventListener(this)
 	, bHasActiveClips(false)
 	, bTimeSynthWasDisabled(false)
 {
 	PrimaryComponentTick.bCanEverTick = true;
+
+	if (TimeSynthCallbackSizeOverrideCvar > 0)
+	{
+		PreferredBufferLength = FMath::Max(TimeSynthCallbackSizeOverrideCvar, DEFAULT_PROCEDURAL_SOUNDWAVE_BUFFER_SIZE);
+	}
 }
 
 UTimeSynthComponent::~UTimeSynthComponent()
@@ -392,7 +421,7 @@ bool UTimeSynthComponent::Init(int32& InSampleRate)
 	SetQuantizationSettings(QuantizationSettings);
 
 	// Create a pool of playing clip runtime infos
-	CurrentPoolSize = MaxPoolSize;
+	CurrentPoolSize = TimeSynthActiveClipLimitCvar;
 
 	PlayingClipsPool_AudioRenderThread.AddDefaulted(CurrentPoolSize);
 	FreePlayingClipIndices_AudioRenderThread.AddDefaulted(CurrentPoolSize);
@@ -835,7 +864,7 @@ FTimeSynthClipHandle UTimeSynthComponent::PlayClip(UTimeSynthClip* InClip, UTime
 	DecodeInit.VolumeScale = VolumeScale;
 	DecodeInit.SoundWave = ChosenSound.SoundWave;
 	DecodeInit.SeekTime = 0;
-	DecodeInit.bForceSyncDecode = true;
+	DecodeInit.bForceSyncDecode = (TimeSynthForceSyncDecodesCvar == 1);
 
 	// Update the synth component on the audio thread
 	FAudioThread::RunCommandOnAudioThread([this, DecodeInit]()
@@ -881,7 +910,6 @@ FTimeSynthClipHandle UTimeSynthComponent::PlayClip(UTimeSynthClip* InClip, UTime
 	// Send this new clip over to the audio render thread
 	SynthCommand([this, NewClipInfo, ClipDuration, FadeInTime, FadeOutTime]
 	{
-		bool bCurrentPoolSizeChanged = false;
 		// Immediately create a mapping for this clip id to a free clip slot
 		// It's possible that the clip might get state changes before it starts playing if
 		// we're playing a very long-duration quantization
@@ -892,28 +920,14 @@ FTimeSynthClipHandle UTimeSynthComponent::PlayClip(UTimeSynthClip* InClip, UTime
 		}
 		else
 		{
-			// Grow the pool size if we ran out of clips in the pool
-			++CurrentPoolSize;
-			bCurrentPoolSizeChanged = true;
-			FreePlayingClipIndices_AudioRenderThread.Add(CurrentPoolSize - 1);
-			FreeClipIndex = FreePlayingClipIndices_AudioRenderThread.Pop(false);
+			UE_LOG(LogTimeSynth, Display, TEXT("Ignoring PlayClip() request since the Playing Clip pool is full, consider initializeng Pool Size to a larger value"));
+			return;
 		}
 		check(FreeClipIndex >= 0);
 
 		// Copy over the clip info to the slot
-		if(bCurrentPoolSizeChanged)
-		{
-			PlayingClipsPool_AudioRenderThread.Add(NewClipInfo);
-			PlayingClipsPool_AudioRenderThread[PlayingClipsPool_AudioRenderThread.Num() - 1].bIsInitialized = true;
-
-			UE_LOG(LogTimeSynth, Warning, TEXT("Reallocating PlayingClipsPool to %i (which is a performance hit.) If this wasn't caused by a hitch, consider initializing Pool Size to a larger value.")
-				, PlayingClipsPool_AudioRenderThread.Num());
-		}
-		else
-		{
-			PlayingClipsPool_AudioRenderThread[FreeClipIndex] = NewClipInfo;
-			PlayingClipsPool_AudioRenderThread[FreeClipIndex].bIsInitialized = true;
-		}
+		PlayingClipsPool_AudioRenderThread[FreeClipIndex] = NewClipInfo;
+		PlayingClipsPool_AudioRenderThread[FreeClipIndex].bIsInitialized = true;
 
 		// Add a mapping of the clip handle id to the free index
 		// This will allow us to reference the playing clip from BP, etc.
