Added BP function to time synth component to query the max active clip limit

diff --git a/Engine/Plugins/Runtime/TimeSynth/Source/TimeSynth/Classes/TimeSynthComponent.h b/Engine/Plugins/Runtime/TimeSynth/Source/TimeSynth/Classes/TimeSynthComponent.h
--- a/Engine/Plugins/Runtime/TimeSynth/Source/TimeSynth/Classes/TimeSynthComponent.h
+++ b/Engine/Plugins/Runtime/TimeSynth/Source/TimeSynth/Classes/TimeSynthComponent.h
@@ -497,6 +497,9 @@ public:
 	UFUNCTION(BlueprintCallable, Category = "TimeSynth")
 	float GetEnvelopeFollowerValue() const { return CurrentEnvelopeValue; }
 
+	UFUNCTION(BlueprintCallable, Category = "TimeSynth")
+	int32 GetMaxActiveClipLimit() const;
+
 	// The default quantizations settings
 	UPROPERTY(EditAnywhere, BlueprintReadWrite, Category = "TimeSynth")
 	FTimeSynthQuantizationSettings QuantizationSettings;
diff --git a/Engine/Plugins/Runtime/TimeSynth/Source/TimeSynth/Private/TimeSynthComponent.cpp b/Engine/Plugins/Runtime/TimeSynth/Source/TimeSynth/Private/TimeSynthComponent.cpp
--- a/Engine/Plugins/Runtime/TimeSynth/Source/TimeSynth/Private/TimeSynthComponent.cpp
+++ b/Engine/Plugins/Runtime/TimeSynth/Source/TimeSynth/Private/TimeSynthComponent.cpp
@@ -1141,3 +1141,8 @@ TArray<FTimeSynthSpectralData> UTimeSynthComponent::GetSpectralData() const
 	return TArray<FTimeSynthSpectralData>();
 }
 
+int32 UTimeSynthComponent::GetMaxActiveClipLimit() const
+{
+	return CurrentPoolSize;
+}
+
