Fixing up uproperty markup for time synth

diff --git a/Engine/Plugins/Runtime/TimeSynth/Source/TimeSynth/Classes/TimeSynthComponent.h b/Engine/Plugins/Runtime/TimeSynth/Source/TimeSynth/Classes/TimeSynthComponent.h
--- a/Engine/Plugins/Runtime/TimeSynth/Source/TimeSynth/Classes/TimeSynthComponent.h
+++ b/Engine/Plugins/Runtime/TimeSynth/Source/TimeSynth/Classes/TimeSynthComponent.h
@@ -529,11 +529,11 @@ public:
 	uint8 bIsFilterBEnabled : 1;
 
 	// The filter settings to use for filter A
-	UPROPERTY(EditAnywhere, BlueprintReadWrite, Category = "Filter", meta = (EditCondition = "bIsFilterEnabled"))
+	UPROPERTY(EditAnywhere, BlueprintReadWrite, Category = "Filter", meta = (EditCondition = "bIsFilterAEnabled"))
 	FTimeSynthFilterSettings FilterASettings;
 
 	// The filter settings to use for filter B
-	UPROPERTY(EditAnywhere, BlueprintReadWrite, Category = "Filter", meta = (EditCondition = "bIsFilterEnabled"))
+	UPROPERTY(EditAnywhere, BlueprintReadWrite, Category = "Filter", meta = (EditCondition = "bIsFilterBEnabled"))
 	FTimeSynthFilterSettings FilterBSettings;
 
 	// Whether or not the filter is enabled
